LWC-related fix for convex hull decomposition. Assume that all input points will always have float precision.
Also remove unnecessary allocations by allowing the async task to assume ownership of the vertex and index arrays passed to it.

diff --git a/Engine/Source/Developer/PhysicsUtilities/Private/ConvexDecompTool.cpp b/Engine/Source/Developer/PhysicsUtilities/Private/ConvexDecompTool.cpp
--- a/Engine/Source/Developer/PhysicsUtilities/Private/ConvexDecompTool.cpp
+++ b/Engine/Source/Developer/PhysicsUtilities/Private/ConvexDecompTool.cpp
@@ -141,7 +141,7 @@ static void InitParameters(IVHACD::Parameters &VHACD_Params, uint32 InHullCount,
 	VHACD_Params.m_pca = 1; // align mesh for more optimal processing of long diagonal meshes
 }
 
-void DecomposeMeshToHulls(UBodySetup* InBodySetup, const TArray<FVector>& InVertices, const TArray<uint32>& InIndices, uint32 InHullCount, int32 InMaxHullVerts,uint32 InResolution)
+void DecomposeMeshToHulls(UBodySetup* InBodySetup, const TArray<FVector3f>& InVertices, const TArray<uint32>& InIndices, uint32 InHullCount, int32 InMaxHullVerts,uint32 InResolution)
 {
 	check(InBodySetup != NULL);
 
@@ -149,7 +149,7 @@ void DecomposeMeshToHulls(UBodySetup* InBodySetup, const TArray<FVector>& InVert
 
 	// Validate input by checking bounding box
 	FBox VertBox(ForceInit);
-	for (FVector Vert : InVertices)
+	for (const FVector3f& Vert : InVertices)
 	{
 		VertBox += Vert;
 	}
@@ -234,34 +234,18 @@ public:
 	{
 	public:
 		VHACDJob(UBodySetup* _InBodySetup,
-				 const TArray<FVector>& InVertices,
-				 const TArray<uint32>& InIndices,
+				 TArray<FVector3f>&& InVertices,
+				 TArray<uint32>&& InIndices,
 				 uint32 InHullCount,
 				 int32 InMaxHullVerts,
-				uint32 InResolution = DEFAULT_HACD_VOXEL_RESOLUTION)
+				 uint32 InResolution = DEFAULT_HACD_VOXEL_RESOLUTION)
+			: Vertices(MoveTemp(InVertices)),			// Take ownership of the buffer of vertex positions
+			  Triangles(MoveTemp(InIndices)),			// Take ownership of the buffer of triangle indices
+			  InBodySetup(_InBodySetup),				// Cache the UBodySetup the convex hull results are to be stored in.
+			  HullCount(InHullCount),					// Save the number of convex hulls requested count
+			  MaxHullVerts(uint32(InMaxHullVerts)),		// Save the maximum number of vertices per convex hull allowed
+			  Resolution(InResolution)					// Save the voxel resolution
 		{
-			InBodySetup = _InBodySetup;						// Cache the UBodySetup the convex hull results are to be stored in.
-			HullCount = InHullCount;						// Save the number of convex hulls requested count
-			MaxHullVerts = uint32(InMaxHullVerts);			// Save the maximum number of vertices per convex hull allowed
-			Resolution = InResolution;						// Save the voxel resolution
-
-			const float* const Verts = (float*)InVertices.GetData();
-			const unsigned int NumVerts = InVertices.Num();
-			const int* const Tris = (int*)InIndices.GetData();
-			const unsigned int NumTris = InIndices.Num() / 3;
-
-			VertCount	= NumVerts;
-			TriCount	= NumTris;
-			// Make a temporary copy of the vertices and triangles for when the job comes due
-			Vertices = new float[NumVerts * 3];
-			memcpy(Vertices, Verts, sizeof(float)*NumVerts * 3);
-			Triangles = new uint32_t[NumTris * 3];
-			memcpy(Triangles, Tris, sizeof(int)*NumTris * 3);
-		}
-
-		~VHACDJob(void)
-		{
-			ReleaseMesh();
 		}
 
 		// Start the V-HACD operation
@@ -270,8 +254,7 @@ public:
 			IVHACD::Parameters VHACD_Params;
 			InitParameters(VHACD_Params, HullCount, MaxHullVerts, Resolution);
 			VHACD_Params.m_callback = InCallback;
-			bool ret = InVHACD->Compute(Vertices, VertCount, Triangles, TriCount, VHACD_Params);
-			ReleaseMesh();
+			bool ret = InVHACD->Compute(&Vertices[0].X, Vertices.Num(), Triangles.GetData(), Triangles.Num() / 3, VHACD_Params);
 			return ret;
 		}
 
@@ -317,24 +300,12 @@ public:
 
 		VHACDJob		*mNext{ nullptr };						// Next job to perform
 private:
-		// Release scratch memory allocated to hold the input request mesh
-		void ReleaseMesh(void)
-		{
-			delete[]Vertices;
-			delete[]Triangles;
-			Vertices = nullptr;
-			Triangles = nullptr;
-		}
-
-
-		float			*Vertices{ nullptr };		// Scratch vertices for the mesh we are operating on
-		uint32_t		*Triangles{ nullptr };		// Scratch triangle indices for the mesh we are operating on
-		uint32			VertCount{ 0 };				// The number of input vertices in the mes
-		uint32			TriCount{ 0 };				// The number of triangles in the mesh
-		UBodySetup		*InBodySetup{ nullptr };	// The body we are going to store the results in
-		uint32			HullCount{ 0 };				// The maximum number of convex hulls requested
-		uint32			MaxHullVerts{ 0 };			// The maximum number of vertices per convex hull requested
-		uint32			Resolution{ 0 };			// The voxel resolution to use for the V-HACD operation
+		TArray<FVector3f>	Vertices;
+		TArray<uint32>		Triangles;
+		UBodySetup*			InBodySetup;		// The body we are going to store the results in
+		uint32				HullCount;			// The maximum number of convex hulls requested
+		uint32				MaxHullVerts;		// The maximum number of vertices per convex hull requested
+		uint32				Resolution;			// The voxel resolution to use for the V-HACD operation
 	};
 
 	FDecomposeMeshToHullsAsyncImpl(void)
@@ -362,7 +333,7 @@ private:
 	*	@param		InAccuracy			Value between 0 and 1, controls how accurate hull generation is
 	*	@param		InMaxHullVerts		Number of verts allowed in a hull
 	*/
-	virtual bool DecomposeMeshToHullsAsyncBegin(UBodySetup* _InBodySetup, const TArray<FVector>& InVertices, const TArray<uint32>& InIndices, uint32 InHullCount, int32 InMaxHullVerts, uint32 InResolution = DEFAULT_HACD_VOXEL_RESOLUTION) final
+	virtual bool DecomposeMeshToHullsAsyncBegin(UBodySetup* _InBodySetup, TArray<FVector3f>&& InVertices, TArray<uint32>&& InIndices, uint32 InHullCount, int32 InMaxHullVerts, uint32 InResolution = DEFAULT_HACD_VOXEL_RESOLUTION) final
 	{
 		check(_InBodySetup != NULL);
 
@@ -370,7 +341,7 @@ private:
 
 		// Validate input by checking bounding box
 		FBox VertBox(ForceInit);
-		for (FVector Vert : InVertices)
+		for (const FVector3f& Vert : InVertices)
 		{
 			VertBox += Vert;
 		}
@@ -382,7 +353,7 @@ private:
 		}
 
 		// Create a VHACD Job instance and add it to the end of the linked list of pending jobs
-		VHACDJob *job = new VHACDJob(_InBodySetup, InVertices, InIndices, InHullCount, InMaxHullVerts, InResolution);
+		VHACDJob *job = new VHACDJob(_InBodySetup, MoveTemp(InVertices), MoveTemp(InIndices), InHullCount, InMaxHullVerts, InResolution);
 		// If we already have an active job, add this job to the end of the list
 		if (CurrentJob)
 		{
diff --git a/Engine/Source/Developer/PhysicsUtilities/Private/PhysicsAssetUtils.cpp b/Engine/Source/Developer/PhysicsUtilities/Private/PhysicsAssetUtils.cpp
--- a/Engine/Source/Developer/PhysicsUtilities/Private/PhysicsAssetUtils.cpp
+++ b/Engine/Source/Developer/PhysicsUtilities/Private/PhysicsAssetUtils.cpp
@@ -537,7 +537,7 @@ bool CreateCollisionFromBoneInternal(UBodySetup* bs, USkeletalMesh* skelMesh, in
 	}
 	else if (Params.GeomType == EFG_SingleConvexHull || Params.GeomType == EFG_MultiConvexHull)
 	{
-		TArray<FVector> Verts;
+		TArray<FVector3f> Verts;
 		TArray<uint32> Indices;
 		TriangleCache.GetVerticesAndIndicesForBone(BoneIndex, Verts, Indices);
 
diff --git a/Engine/Source/Developer/PhysicsUtilities/Private/SkinnedBoneTriangleCache.cpp b/Engine/Source/Developer/PhysicsUtilities/Private/SkinnedBoneTriangleCache.cpp
--- a/Engine/Source/Developer/PhysicsUtilities/Private/SkinnedBoneTriangleCache.cpp
+++ b/Engine/Source/Developer/PhysicsUtilities/Private/SkinnedBoneTriangleCache.cpp
@@ -24,7 +24,7 @@ void FSkinnedBoneTriangleCache::BuildCache()
 	BuildOwnedTrianglesSetForEachBone();
 }
 
-void FSkinnedBoneTriangleCache::GetVerticesAndIndicesForBone(const int32 BoneIndex, TArray<FVector>& OutVertexPositions, TArray<uint32>& OutIndices) const
+void FSkinnedBoneTriangleCache::GetVerticesAndIndicesForBone(const int32 BoneIndex, TArray<FVector3f>& OutVertexPositions, TArray<uint32>& OutIndices) const
 {
 	OutVertexPositions.Empty();
 	OutIndices.Empty();
@@ -155,9 +155,9 @@ void FSkinnedBoneTriangleCache::AddTriangleIndexToOwnerBoneSets(const FTriangleI
 	}
 }
 
-FVector FSkinnedBoneTriangleCache::VertexPosition(const FSkinnedVertexIndex VertIndex, const FMatrix& ComponentToBoneMatrix) const
+FVector3f FSkinnedBoneTriangleCache::VertexPosition(const FSkinnedVertexIndex VertIndex, const FMatrix& ComponentToBoneMatrix) const
 {
-	const FVector Position = VertexBuffer.VertexPosition(VertIndex);
+	const FVector3f& Position = VertexBuffer.VertexPosition(VertIndex);
 	return ComponentToBoneMatrix.TransformPosition(Position);
 }
 
diff --git a/Engine/Source/Developer/PhysicsUtilities/Private/SkinnedBoneTriangleCache.h b/Engine/Source/Developer/PhysicsUtilities/Private/SkinnedBoneTriangleCache.h
--- a/Engine/Source/Developer/PhysicsUtilities/Private/SkinnedBoneTriangleCache.h
+++ b/Engine/Source/Developer/PhysicsUtilities/Private/SkinnedBoneTriangleCache.h
@@ -33,7 +33,7 @@ public:
 	explicit FSkinnedBoneTriangleCache(USkeletalMesh& InSkeletalMesh, const FPhysAssetCreateParams& Params);
 
 	void BuildCache();
-	void GetVerticesAndIndicesForBone(const int32 BoneIndex, TArray<FVector>& OutVertexPositions, TArray<uint32>& OutIndices) const;
+	void GetVerticesAndIndicesForBone(const int32 BoneIndex, TArray<FVector3f>& OutVertexPositions, TArray<uint32>& OutIndices) const;
 
 private:
 	typedef int32 FBoneIndex;
@@ -70,7 +70,7 @@ private:
 	/** Here, each triangle index is added to its appropriate sets. */
 	void AddTriangleIndexToOwnerBoneSets(const FTriangleIndex TriangleIndex);
 
-	FVector VertexPosition(const FSkinnedVertexIndex VertIndex, const FMatrix& ComponentToBoneMatrix) const;
+	FVector3f VertexPosition(const FSkinnedVertexIndex VertIndex, const FMatrix& ComponentToBoneMatrix) const;
 
 	// Inputs
 	USkeletalMesh& SkeletalMesh;
diff --git a/Engine/Source/Developer/PhysicsUtilities/Public/ConvexDecompTool.h b/Engine/Source/Developer/PhysicsUtilities/Public/ConvexDecompTool.h
--- a/Engine/Source/Developer/PhysicsUtilities/Public/ConvexDecompTool.h
+++ b/Engine/Source/Developer/PhysicsUtilities/Public/ConvexDecompTool.h
@@ -29,7 +29,7 @@ class UBodySetup;
  *	@param		InAccuracy			Value between 0 and 1, controls how accurate hull generation is
  *	@param		InMaxHullVerts		Number of verts allowed in a hull
  */
-PHYSICSUTILITIES_API void DecomposeMeshToHulls(UBodySetup* InBodySetup, const TArray<FVector>& InVertices, const TArray<uint32>& InIndices, uint32 InHullCount, int32 InMaxHullVerts,uint32 InResolution= DEFAULT_HACD_VOXEL_RESOLUTION);
+PHYSICSUTILITIES_API void DecomposeMeshToHulls(UBodySetup* InBodySetup, const TArray<FVector3f>& InVertices, const TArray<uint32>& InIndices, uint32 InHullCount, int32 InMaxHullVerts,uint32 InResolution= DEFAULT_HACD_VOXEL_RESOLUTION);
 
 class IDecomposeMeshToHullsAsync
 {
@@ -40,12 +40,12 @@ public:
 	/**
 	*	Utility for turning arbitrary mesh into convex hulls.  This can be called multiple times for multiple UBodySetup results and they will be done processed in sequential order
 	*	@output		InBodySetup			BodySetup that will have its existing hulls removed and replaced with results of decomposition.
-	*	@param		InVertices			Array of vertex positions of input mesh
-	*	@param		InIndices			Array of triangle indices for input mesh
+	*	@param		InVertices			Array of vertex positions of input mesh, whose ownership is transferred to the async task
+	*	@param		InIndices			Array of triangle indices for input mesh, whose ownership is transferred to the async task
 	*	@param		InAccuracy			Value between 0 and 1, controls how accurate hull generation is
 	*	@param		InMaxHullVerts		Number of verts allowed in a hull
 	*/
-	virtual bool DecomposeMeshToHullsAsyncBegin(UBodySetup* InBodySetup, const TArray<FVector>& InVertices, const TArray<uint32>& InIndices, uint32 InHullCount, int32 InMaxHullVerts, uint32 InResolution = DEFAULT_HACD_VOXEL_RESOLUTION) = 0;
+	virtual bool DecomposeMeshToHullsAsyncBegin(UBodySetup* InBodySetup, TArray<FVector3f>&& InVertices, TArray<uint32>&& InIndices, uint32 InHullCount, int32 InMaxHullVerts, uint32 InResolution = DEFAULT_HACD_VOXEL_RESOLUTION) = 0;
 
 	// Returns true if the convex decomposition process has completed and results are ready.
 	virtual bool IsComplete(void) = 0;
diff --git a/Engine/Source/Editor/StaticMeshEditor/Private/StaticMeshEditor.cpp b/Engine/Source/Editor/StaticMeshEditor/Private/StaticMeshEditor.cpp
--- a/Engine/Source/Editor/StaticMeshEditor/Private/StaticMeshEditor.cpp
+++ b/Engine/Source/Editor/StaticMeshEditor/Private/StaticMeshEditor.cpp
@@ -2050,11 +2050,12 @@ void FStaticMeshEditor::DoDecomp(uint32 InHullCount, int32 InMaxHullVerts, uint3
 
 		// Make vertex buffer
 		int32 NumVerts = LODModel.VertexBuffers.StaticMeshVertexBuffer.GetNumVertices();
-		TArray<FVector> Verts;
+		TArray<FVector3f> Verts;
+		Verts.SetNumUninitialized(NumVerts);
 		for(int32 i=0; i<NumVerts; i++)
 		{
-			FVector Vert = LODModel.VertexBuffers.PositionVertexBuffer.VertexPosition(i);
-			Verts.Add(Vert);
+			const FVector3f& Vert = LODModel.VertexBuffers.PositionVertexBuffer.VertexPosition(i);
+			Verts[i] = Vert;
 		}
 
 		// Grab all indices
@@ -2103,7 +2104,7 @@ void FStaticMeshEditor::DoDecomp(uint32 InHullCount, int32 InMaxHullVerts, uint3
 			}
 			// Begin the convex decomposition process asynchronously
 			DecomposeMeshToHullsAsync = CreateIDecomposeMeshToHullAsync();
-			DecomposeMeshToHullsAsync->DecomposeMeshToHullsAsyncBegin(bs, Verts, CollidingIndices, InHullCount, InMaxHullVerts, InHullPrecision);
+			DecomposeMeshToHullsAsync->DecomposeMeshToHullsAsyncBegin(bs, MoveTemp(Verts), MoveTemp(CollidingIndices), InHullCount, InMaxHullVerts, InHullPrecision);
 #else
 			DecomposeMeshToHulls(bs, Verts, CollidingIndices, InHullCount, InMaxHullVerts, InHullPrecision);
 #endif
diff --git a/Engine/Source/Editor/StaticMeshEditor/Private/StaticMeshEditorSubsystem.cpp b/Engine/Source/Editor/StaticMeshEditor/Private/StaticMeshEditorSubsystem.cpp
--- a/Engine/Source/Editor/StaticMeshEditor/Private/StaticMeshEditorSubsystem.cpp
+++ b/Engine/Source/Editor/StaticMeshEditor/Private/StaticMeshEditorSubsystem.cpp
@@ -69,7 +69,7 @@ namespace InternalEditorMeshLibrary
 
 		// Make vertex buffer
 		int32 NumVerts = LODModel.VertexBuffers.StaticMeshVertexBuffer.GetNumVertices();
-		TArray<FVector> Verts;
+		TArray<FVector3f> Verts;
 		Verts.Reserve(NumVerts);
 		for (int32 i = 0; i < NumVerts; i++)
 		{
