Warning added about capturing a probably-not-trivially-relocatable std::function and handing it off to UE.

diff --git a/Engine/Source/Developer/PhysicsUtilities/Private/ConvexDecompTool.cpp b/Engine/Source/Developer/PhysicsUtilities/Private/ConvexDecompTool.cpp
--- a/Engine/Source/Developer/PhysicsUtilities/Private/ConvexDecompTool.cpp
+++ b/Engine/Source/Developer/PhysicsUtilities/Private/ConvexDecompTool.cpp
@@ -85,6 +85,10 @@ public:
 
 	virtual void* StartTask(std::function<void()> func)
 	{
+		// Capturing a std::function and passing it to Async.
+		// This works today but if Async() or TFuture start to assume they can trivially relocate
+		// callbacks (e.g. by storing them in a TArray) then this capture will need revising.
+		// Unlike TFunction, std::function is not guaranteed to be trivially relocatable.
 		return new TFuture<void>(Async(EAsyncExecution::ThreadPool, [func](){func();}));
 	}
 
