SkeletalMeshEditor: add option to use level sets when running "Generate All Bodies"

diff --git a/Engine/Source/Developer/PhysicsUtilities/PhysicsUtilities.Build.cs b/Engine/Source/Developer/PhysicsUtilities/PhysicsUtilities.Build.cs
--- a/Engine/Source/Developer/PhysicsUtilities/PhysicsUtilities.Build.cs
+++ b/Engine/Source/Developer/PhysicsUtilities/PhysicsUtilities.Build.cs
@@ -20,6 +20,7 @@ public class PhysicsUtilities : ModuleRules
 			new string[] {
 				"MeshUtilitiesCommon",
 				"MeshUtilitiesEngine",
+				"GeometryCore"		// For mesh to level set conversion
             }
 		);
 
diff --git a/Engine/Source/Developer/PhysicsUtilities/Private/PhysicsAssetUtils.cpp b/Engine/Source/Developer/PhysicsUtilities/Private/PhysicsAssetUtils.cpp
--- a/Engine/Source/Developer/PhysicsUtilities/Private/PhysicsAssetUtils.cpp
+++ b/Engine/Source/Developer/PhysicsUtilities/Private/PhysicsAssetUtils.cpp
@@ -18,6 +18,9 @@
 #include "PreviewScene.h"
 #include "Misc/ScopedSlowTask.h"
 #include "SkinnedBoneTriangleCache.h"
+#include "DynamicMesh/DynamicMesh3.h"		// For mesh to level set
+#include "Implicit/SweepingMeshSDF.h"		// For mesh to level set
+#include "Implicit/Solidify.h"		// For mesh to level set
 
 namespace FPhysicsAssetUtils
 {
@@ -345,7 +348,7 @@ bool CreateFromSkeletalMesh(UPhysicsAsset* PhysicsAsset, USkeletalMesh* SkelMesh
 
 	FSkinnedBoneTriangleCache TriangleCache(*SkelMesh, Params);
 
-	if (Params.GeomType == EFG_SingleConvexHull || Params.GeomType == EFG_MultiConvexHull)
+	if (Params.GeomType == EFG_SingleConvexHull || Params.GeomType == EFG_MultiConvexHull || Params.GeomType == EFG_LevelSet)
 	{
 		TriangleCache.BuildCache();
 	}
@@ -440,11 +443,143 @@ FVector ComputeEigenVector(const FMatrix& A)
 	return Bk.GetSafeNormal();
 }
 
+
+namespace LevelSetHelpers
+{
+	// Copied from CollisionGeometryConversion.h (Plugins/Runtime/MeshModelingToolset)
+	void CreateLevelSetElement(const FTransform3d& GridTransform, const UE::Geometry::FDenseGrid3f& Grid, float CellSize, FKLevelSetElem& LevelSetOut)
+	{
+		const UE::Geometry::FVector3i& GridDims = Grid.GetDimensions();
+		TArray<double> OutGridValues;
+		OutGridValues.Init(0.0, GridDims[0] * GridDims[1] * GridDims[2]);
+
+		for (int I = 0; I < GridDims[0]; ++I)
+		{
+			for (int J = 0; J < GridDims[1]; ++J)
+			{
+				for (int K = 0; K < GridDims[2]; ++K)
+				{
+					// account for different ordering in Geometry::TDenseGrid3 vs Chaos::TUniformGrid
+					const int InBufferIndex = I + GridDims[0] * (J + GridDims[1] * K);
+					const int OutBufferIndex = K + GridDims[2] * (J + GridDims[1] * I);
+					OutGridValues[OutBufferIndex] = Grid[InBufferIndex];
+				}
+			}
+		}
+
+		FTransform ChaosTransform = GridTransform;
+		ChaosTransform.AddToTranslation(-0.5 * CellSize * FVector::One());		// account for grid origin being at cell center vs cell corner
+		LevelSetOut.BuildLevelSet(ChaosTransform, OutGridValues, FIntVector(GridDims[0], GridDims[1], GridDims[2]), CellSize);
+	}
+
+	// Copied from FMeshSimpleShapeApproximation::Generate_LevelSets (Plugins/Runtime/GeometryProcessing)
+	bool CreateLevelSetForMesh(const UE::Geometry::FDynamicMesh3& InMesh, int32 LevelSetGridResolution, FKLevelSetElem& OutElement)
+	{
+		using UE::Geometry::FDynamicMesh3;
+		using UE::Geometry::TSweepingMeshSDF;
+
+		const UE::Geometry::FAxisAlignedBox3d Bounds = InMesh.GetBounds();
+		const double CellSize = Bounds.MaxDim() / LevelSetGridResolution;
+		const double ExpandBounds = 2.0 * CellSize;
+
+		UE::Geometry::TMeshAABBTree3<FDynamicMesh3> Spatial(&InMesh);
+
+		// Input mesh is likely open, so solidify it first
+		UE::Geometry::TFastWindingTree<FDynamicMesh3> FastWinding(&Spatial);
+		UE::Geometry::TImplicitSolidify<FDynamicMesh3> Solidify(&InMesh, &Spatial, &FastWinding);
+		Solidify.ExtendBounds = ExpandBounds;
+		Solidify.MeshCellSize = CellSize;
+		Solidify.WindingThreshold = 0.5;
+		Solidify.SurfaceSearchSteps = 3;
+		Solidify.bSolidAtBoundaries = true;
+
+		const UE::Geometry::FDynamicMesh3 SolidMesh = FDynamicMesh3(&Solidify.Generate());
+		Spatial.SetMesh(&SolidMesh, true);
+
+		TSweepingMeshSDF<FDynamicMesh3> SDF;
+		SDF.Mesh = &SolidMesh;
+		SDF.Spatial = &Spatial;
+		SDF.ComputeMode = TSweepingMeshSDF<FDynamicMesh3>::EComputeModes::NarrowBand_SpatialFloodFill;
+		SDF.CellSize = (float)CellSize;
+		SDF.NarrowBandMaxDistance = 2.0 * CellSize;
+		SDF.ExactBandWidth = FMath::CeilToInt32(SDF.NarrowBandMaxDistance / CellSize);
+		SDF.ExpandBounds = FVector3d(ExpandBounds);
+
+		if (SDF.Compute(Bounds))
+		{
+			const FTransform3d GridTransform = FTransform((FVector3d)SDF.GridOrigin);
+			CreateLevelSetElement(GridTransform, SDF.Grid, SDF.CellSize, OutElement);
+			return true;
+		}
+
+		return false;
+	}
+
+	void CreateDynamicMesh(const TArray<FVector3f>& InVertices, const TArray<uint32>& InIndices, UE::Geometry::FDynamicMesh3& OutMesh)
+	{
+		OutMesh.Clear();
+		for (const FVector3f& Vertex : InVertices)
+		{
+			OutMesh.AppendVertex(FVector3d(Vertex));
+		}
+
+		check(InIndices.Num() % 3 == 0);
+		const int32 NumTriangles = InIndices.Num() / 3;
+		for (int32 TriangleID = 0; TriangleID < NumTriangles; ++TriangleID)
+		{
+			const UE::Geometry::FIndex3i Triangle(InIndices[3 * TriangleID], InIndices[3 * TriangleID + 1], InIndices[3 * TriangleID + 2]);
+			OutMesh.AppendTriangle(Triangle);
+		}
+	}
+
+	bool CreateLevelSetForBone(UBodySetup* BodySetup, const TArray<FVector3f>& InVertices, const TArray<uint32>& InIndices, uint32 InResolution)
+	{
+		check(BodySetup != NULL);
+
+		// Validate input by checking bounding box
+		FBox VertBox(ForceInit);
+		for (const FVector3f& Vert : InVertices)
+		{
+			VertBox += (FVector)Vert;
+		}
+
+		// If box is invalid, or the largest dimension is less than 1 unit, or smallest is less than 0.1, skip trying to generate collision
+		if (VertBox.IsValid == 0 || VertBox.GetSize().GetMax() < 1.f || VertBox.GetSize().GetMin() < 0.1f || InIndices.Num() == 0)
+		{
+			return false;
+		}
+
+		TRACE_CPUPROFILER_EVENT_SCOPE(MeshToLevelSet)
+
+		// Clean out old hulls
+		BodySetup->RemoveSimpleCollision();
+
+		UE::Geometry::FDynamicMesh3 DynamicMesh;
+		CreateDynamicMesh(InVertices, InIndices, DynamicMesh);
+
+		FKLevelSetElem LevelSetElement;
+		const bool bOK = CreateLevelSetForMesh(DynamicMesh, InResolution, LevelSetElement);
+
+		if (!bOK)
+		{
+			return false; 
+		}
+
+		BodySetup->AggGeom.LevelSetElems.Add(LevelSetElement);
+		BodySetup->InvalidatePhysicsData(); // update GUID
+		return true;
+	}
+
+}
+
+
+
 bool CreateCollisionFromBoneInternal(UBodySetup* bs, USkeletalMesh* skelMesh, int32 BoneIndex, const FPhysAssetCreateParams& Params, const FBoneVertInfo& Info, const FSkinnedBoneTriangleCache& TriangleCache)
 {
 #if WITH_EDITOR
 	// multi convex hull can fail so wait to clear it ( will be called in DecomposeMeshToHulls() )
-	if (Params.GeomType != EFG_SingleConvexHull && Params.GeomType != EFG_MultiConvexHull)
+	const bool bCanFail = (Params.GeomType == EFG_SingleConvexHull || Params.GeomType == EFG_MultiConvexHull || Params.GeomType == EFG_LevelSet);
+	if (!bCanFail)
 	{
 		// Empty any existing collision.
 		bs->RemoveSimpleCollision();
@@ -604,6 +739,31 @@ bool CreateCollisionFromBoneInternal(UBodySetup* bs, USkeletalMesh* skelMesh, in
 
 		bs->AggGeom.TaperedCapsuleElems.Add(TaperedCapsuleElem);
 	}
+	else if (Params.GeomType == EFG_LevelSet)
+	{
+		TArray<FVector3f> Verts;
+		TArray<uint32> Indices;
+		TriangleCache.GetVerticesAndIndicesForBone(BoneIndex, Verts, Indices);
+
+		if (Verts.Num())
+		{
+			const bool bOK = LevelSetHelpers::CreateLevelSetForBone(bs, Verts, Indices, Params.LevelSetResolution);
+			if (!bOK)
+			{
+				FMessageLog EditorErrors("EditorErrors");
+				EditorErrors.Warning(NSLOCTEXT("PhysicsAssetUtils", "LevelSetError", "An error occurred creating a level set for the given bone."));
+				EditorErrors.Open();
+				return false;
+			}
+		}
+		else
+		{
+			FMessageLog EditorErrors("EditorErrors");
+			EditorErrors.Warning(NSLOCTEXT("PhysicsAssetUtils", "LevelSetNoPositions", "Unable to create a level set for the given bone as there are no vertices associated with the bone."));
+			EditorErrors.Open();
+			return false;
+		}
+	}
 
 	return true;
 }
@@ -614,7 +774,7 @@ bool CreateCollisionFromBone(UBodySetup* bs, USkeletalMesh* skelMesh, int32 Bone
 
 	FSkinnedBoneTriangleCache TriangleCache(*skelMesh, Params);
 
-	if (Params.GeomType == EFG_SingleConvexHull || Params.GeomType == EFG_MultiConvexHull)
+	if (Params.GeomType == EFG_SingleConvexHull || Params.GeomType == EFG_MultiConvexHull || Params.GeomType == EFG_LevelSet)
 	{
 		TriangleCache.BuildCache();
 	}
@@ -628,7 +788,7 @@ bool CreateCollisionFromBones(UBodySetup* bs, USkeletalMesh* skelMesh, const TAr
 
 	FSkinnedBoneTriangleCache TriangleCache(*skelMesh, Params);
 
-	if (Params.GeomType == EFG_SingleConvexHull || Params.GeomType == EFG_MultiConvexHull)
+	if (Params.GeomType == EFG_SingleConvexHull || Params.GeomType == EFG_MultiConvexHull || Params.GeomType == EFG_LevelSet)
 	{
 		TriangleCache.BuildCache();
 	}
diff --git a/Engine/Source/Developer/PhysicsUtilities/Public/PhysicsAssetUtils.h b/Engine/Source/Developer/PhysicsUtilities/Public/PhysicsAssetUtils.h
--- a/Engine/Source/Developer/PhysicsUtilities/Public/PhysicsAssetUtils.h
+++ b/Engine/Source/Developer/PhysicsUtilities/Public/PhysicsAssetUtils.h
@@ -21,7 +21,8 @@ enum EPhysAssetFitGeomType
 	EFG_Sphere				UMETA(DisplayName="Sphere"),
 	EFG_TaperedCapsule		UMETA(DisplayName="Tapered Capsule (Cloth Only)"),
 	EFG_SingleConvexHull	UMETA(DisplayName="Single Convex Hull"),
-	EFG_MultiConvexHull		UMETA(DisplayName="Multi Convex Hull")
+	EFG_MultiConvexHull		UMETA(DisplayName="Multi Convex Hull"),
+	EFG_LevelSet			UMETA(DisplayName="Level Set"),
 };
 
 UENUM()
@@ -51,6 +52,7 @@ struct FPhysAssetCreateParams
 		AngularConstraintMode = ACM_Limited;
 		HullCount = 4;
 		MaxHullVerts = 16;
+		LevelSetResolution = 8;
 	}
 
 	/** Bones that are shorter than this value will be ignored for body creation */
@@ -100,6 +102,11 @@ struct FPhysAssetCreateParams
 	/** When creating convex hulls, the maximum verts that should be created */
 	UPROPERTY(EditAnywhere, AdvancedDisplay, Category = "Body Creation")
 	int32								MaxHullVerts;
+
+	/** When creating level sets, the grid resolution to use */
+	UPROPERTY(EditAnywhere, AdvancedDisplay, Category = "Body Creation", 
+		meta = (ClampMin = 1, UIMin = 5, UIMax = 100, ClampMax = 500, EditCondition = "GeomType == EPhysAssetFitGeomType::EFG_LevelSet"))
+	int32								LevelSetResolution;
 };
 
 class UPhysicsAsset;
diff --git a/Engine/Source/Editor/PhysicsAssetEditor/Private/PhysicsAssetEditorSharedData.cpp b/Engine/Source/Editor/PhysicsAssetEditor/Private/PhysicsAssetEditorSharedData.cpp
--- a/Engine/Source/Editor/PhysicsAssetEditor/Private/PhysicsAssetEditorSharedData.cpp
+++ b/Engine/Source/Editor/PhysicsAssetEditor/Private/PhysicsAssetEditorSharedData.cpp
@@ -1433,6 +1433,7 @@ static EAggCollisionShape::Type ConvertPhysicsAssetGeomTypeToAggCollisionShapeTy
 	case EPhysAssetFitGeomType::EFG_TaperedCapsule: 	return EAggCollisionShape::Type::TaperedCapsule;
 	case EPhysAssetFitGeomType::EFG_SingleConvexHull:	return EAggCollisionShape::Type::Convex;
 	case EPhysAssetFitGeomType::EFG_MultiConvexHull:	return EAggCollisionShape::Type::Convex;
+	case EPhysAssetFitGeomType::EFG_LevelSet:			return EAggCollisionShape::Type::LevelSet;
 	default:											return EAggCollisionShape::Type::Unknown;
 	}
 }
