SkinnedBoneTriangleCache: fix for getting the wrong bone-vertex influence.

Weights are now using uint16 but one variable was still uint8

diff --git a/Engine/Source/Developer/PhysicsUtilities/Private/SkinnedBoneTriangleCache.cpp b/Engine/Source/Developer/PhysicsUtilities/Private/SkinnedBoneTriangleCache.cpp
--- a/Engine/Source/Developer/PhysicsUtilities/Private/SkinnedBoneTriangleCache.cpp
+++ b/Engine/Source/Developer/PhysicsUtilities/Private/SkinnedBoneTriangleCache.cpp
@@ -81,7 +81,7 @@ void FSkinnedBoneTriangleCache::AddIndexToInfluencerBoneSets(const FSkinnedVerte
 
 	const FSkelMeshSection& Section = StaticLODModel.LODModels[0].Sections[SectionIndex];
 	const FSoftSkinVertex& SoftVert = Section.SoftVertices[SoftVertIndex];
-	const uint8 MaxWeight = InfluenceHeuristic == EVW_DominantWeight ? SoftVert.GetMaximumWeight() : 0;
+	const uint16 MaxWeight = InfluenceHeuristic == EVW_DominantWeight ? SoftVert.GetMaximumWeight() : 0;
 
 	for (int32 InfluenceIndex = 0; InfluenceIndex < MAX_TOTAL_INFLUENCES; ++InfluenceIndex)
 	{
