Use the "FullGrid" option when generating LevelSet physics assets.

Chaos cloth (the primary user of level sets) can penetrate arbitrarily deep inside the collider and still needs a good Phi and Normal value, rather than only having accurate values within a narrow band. This makes calculating the levelset more expensive (a one time cost), but in practice, it's not really noticeable.

diff --git a/Engine/Source/Developer/PhysicsUtilities/Private/PhysicsAssetUtils.cpp b/Engine/Source/Developer/PhysicsUtilities/Private/PhysicsAssetUtils.cpp
--- a/Engine/Source/Developer/PhysicsUtilities/Private/PhysicsAssetUtils.cpp
+++ b/Engine/Source/Developer/PhysicsUtilities/Private/PhysicsAssetUtils.cpp
@@ -506,7 +506,7 @@ namespace LevelSetHelpers
 		TSweepingMeshSDF<FDynamicMesh3> SDF;
 		SDF.Mesh = &SolidMesh;
 		SDF.Spatial = &Spatial;
-		SDF.ComputeMode = TSweepingMeshSDF<FDynamicMesh3>::EComputeModes::NarrowBand_SpatialFloodFill;
+		SDF.ComputeMode = TSweepingMeshSDF<FDynamicMesh3>::EComputeModes::FullGrid;
 		SDF.CellSize = (float)CellSize;
 		SDF.NarrowBandMaxDistance = NumNarrowBandCells * CellSize;
 		SDF.ExactBandWidth = FMath::CeilToInt32(SDF.NarrowBandMaxDistance / CellSize);
