Fix unity build

diff --git a/Engine/Source/Developer/PhysicsUtilities/Private/LevelSetHelpers.h b/Engine/Source/Developer/PhysicsUtilities/Private/LevelSetHelpers.h
--- a/Engine/Source/Developer/PhysicsUtilities/Private/LevelSetHelpers.h
+++ b/Engine/Source/Developer/PhysicsUtilities/Private/LevelSetHelpers.h
@@ -1,8 +1,11 @@
 // Copyright Epic Games, Inc. All Rights Reserved.
 
 #pragma once
+#include "CoreTypes.h"
+#include "Templates/UniquePtr.h"
 
 struct FKLevelSetElem;
+class UBodySetup;
 namespace UE::Geometry
 {
 class FDynamicMesh3;
@@ -14,7 +17,7 @@ class FLevelSet;
 
 namespace LevelSetHelpers
 {
-bool CreateLevelSetForBone(class UBodySetup* BodySetup, const TArray<FVector3f>& InVertices, const TArray<uint32>& InIndices, uint32 InResolution);
+bool CreateLevelSetForBone(UBodySetup* BodySetup, const TArray<FVector3f>& InVertices, const TArray<uint32>& InIndices, uint32 InResolution);
 void CreateDynamicMesh(const TArray<FVector3f>& InVertices, const TArray<uint32>& InIndices, UE::Geometry::FDynamicMesh3& OutMesh);
 bool CreateLevelSetForMesh(const UE::Geometry::FDynamicMesh3& InMesh, int32 InLevelSetGridResolution, FKLevelSetElem& OutElement);
 bool CreateLevelSetForMesh(const UE::Geometry::FDynamicMesh3& InMesh, int32 InLevelSetGridResolution, TUniquePtr<Chaos::FLevelSet>& OutElement);
diff --git a/Engine/Source/Developer/PhysicsUtilities/Private/SkinnedLevelSetBuilder.h b/Engine/Source/Developer/PhysicsUtilities/Private/SkinnedLevelSetBuilder.h
--- a/Engine/Source/Developer/PhysicsUtilities/Private/SkinnedLevelSetBuilder.h
+++ b/Engine/Source/Developer/PhysicsUtilities/Private/SkinnedLevelSetBuilder.h
@@ -8,6 +8,7 @@ class FSkinnedBoneTriangleCache;
 class FSkeletalMeshModel;
 class FSkeletalMeshRenderData;
 struct FPhysAssetCreateParams;
+class USkeletalMesh;
 namespace Chaos
 {
 class FLevelSet;
