Reduce number of bad weights in SkinnedLevelSet that cause "spikes". These spikes are caused by vertices getting small weights that get renormalized to large weights, in particular because for some reason the largest influence is not set.

Firstly, improve the dirichlet weights on the surface when doing the poisson diffusion by only considering vertices with non-zero weights for the closest point calculation. Before, it was possible to get zero weights because the triangle mesh includes any triangle with non-zero weights.

Secondly, when renormalizing weights, set a higher threshold (0.1) for discarding weights.

Also added debug draw of the weights.

diff --git a/Engine/Source/Developer/PhysicsUtilities/Private/SkinnedLevelSetBuilder.cpp b/Engine/Source/Developer/PhysicsUtilities/Private/SkinnedLevelSetBuilder.cpp
--- a/Engine/Source/Developer/PhysicsUtilities/Private/SkinnedLevelSetBuilder.cpp
+++ b/Engine/Source/Developer/PhysicsUtilities/Private/SkinnedLevelSetBuilder.cpp
@@ -7,6 +7,7 @@
 #include "BoneWeights.h"
 #include "Chaos/ArrayND.h"
 #include "Chaos/Levelset.h"
+#include "Chaos/Plane.h"
 #include "Chaos/UniformGrid.h"
 #include "Chaos/Math/Poisson.h"
 #include "DynamicMesh/DynamicMesh3.h"
@@ -75,6 +76,58 @@ bool FSkinnedLevelSetBuilder::InitializeSkinnedLevelset(const FPhysAssetCreatePa
 	return false;
 }
 
+static float CalculateTriangleWeight(const UE::Geometry::FTriangle3d& Triangle, const FVector3f& TriangleWeights, const Chaos::TVector<double, 3>& Location)
+{
+	if (TriangleWeights.X == 0.f)
+	{
+		if (TriangleWeights.Y == 0.f)
+		{
+			check(TriangleWeights.Z != 0.f);
+			return TriangleWeights.Z;
+		}
+		else if (TriangleWeights.Z == 0.f)
+		{
+			check(TriangleWeights.Y != 0.f);
+			return TriangleWeights.Y;
+		}
+		else
+		{
+			double Alpha;
+			Chaos::FindClosestPointAndAlphaOnLineSegment(Chaos::FVec3(Triangle.V[1]), Chaos::FVec3(Triangle.V[2]), Location, Alpha);
+			return (1. - Alpha) * TriangleWeights.Y + Alpha * TriangleWeights.Z;
+		}
+	}
+	else if (TriangleWeights.Y == 0.f)
+	{
+		if (TriangleWeights.Z == 0.f)
+		{
+			check(TriangleWeights.X != 0.f);
+			return TriangleWeights.X;
+		}
+		else
+		{
+			double Alpha;
+			Chaos::FindClosestPointAndAlphaOnLineSegment(Chaos::FVec3(Triangle.V[0]), Chaos::FVec3(Triangle.V[2]), Location, Alpha);
+			return (1. - Alpha) * TriangleWeights.X + Alpha * TriangleWeights.Z;
+		}
+	}
+	else if (TriangleWeights.Z == 0.f)
+	{
+		double Alpha;
+		Chaos::FindClosestPointAndAlphaOnLineSegment(Chaos::FVec3(Triangle.V[0]), Chaos::FVec3(Triangle.V[1]), Location, Alpha);
+		return (1. - Alpha) * TriangleWeights.X + Alpha * TriangleWeights.Y;
+
+	}
+	else
+	{
+		// Get Barycentric coordinates for closest point
+		UE::Geometry::TDistPoint3Triangle3<double> TriQuery(Location, Triangle);
+		TriQuery.ComputeResult();
+		return TriQuery.TriangleBaryCoords.X * TriangleWeights.X +
+			TriQuery.TriangleBaryCoords.Y * TriangleWeights.Y + TriQuery.TriangleBaryCoords.Z * TriangleWeights.Z;
+	}
+}
+
 void FSkinnedLevelSetBuilder::AddBoneInfluence(int32 PrimaryBoneIndex, const TArray<int32>& AllBonesForInfluence)
 {
 	// Need to generate weights for grid points corresponding with this bone.
@@ -144,7 +197,7 @@ void FSkinnedLevelSetBuilder::AddBoneInfluence(int32 PrimaryBoneIndex, const TAr
 	const UE::Geometry::FDynamicMesh3 SolidMesh(&Solidify.Generate());
 	UE::Geometry::TMeshAABBTree3<UE::Geometry::FDynamicMesh3> SolidSpatial(&SolidMesh);
 	UE::Geometry::TFastWindingTree<UE::Geometry::FDynamicMesh3> SolidFastWinding(&SolidSpatial);
-	constexpr int32 OuterWeightExpandCells = 3; // Diffuse out this far outside mesh
+	constexpr int32 OuterWeightExpandCells = 4; // Diffuse out this far outside mesh
 	constexpr int32 SurfaceExpandCells = 1; // Consider points this far from the surface of the mesh to be on the surface (dirichlet boundary)
 	UE::Geometry::IMeshSpatial::FQueryOptions QueryOptions;
 	QueryOptions.MaxDistance = (double)LatticeGrid.Dx().Max() * FMath::Max(OuterWeightExpandCells, SurfaceExpandCells);// Used for initial query against solidified mesh
@@ -203,14 +256,13 @@ void FSkinnedLevelSetBuilder::AddBoneInfluence(int32 PrimaryBoneIndex, const TAr
 				// Need to find closest triangle on original mesh to get weights. No distance limit since we might be in a hole that was filled
 				// by solidify. 
 				const int32 NearestTriangle = Spatial.FindNearestTriangle(Location, NearestDistSq);
-				// Get Barycentric coordinates for closest point
+				const UE::Geometry::FIndex3i& TriangleIndices = DynamicMesh.GetTriangle(NearestTriangle);
+				const FVector3f TriangleWeights(GetWeightForIndices(BonesSet, OrigIndices[TriangleIndices.A]), GetWeightForIndices(BonesSet, OrigIndices[TriangleIndices.B]), GetWeightForIndices(BonesSet, OrigIndices[TriangleIndices.C]));
+
 				UE::Geometry::FTriangle3d Triangle;
 				DynamicMesh.GetTriVertices(NearestTriangle, Triangle.V[0], Triangle.V[1], Triangle.V[2]);
-				UE::Geometry::TDistPoint3Triangle3<double> TriQuery(Location, Triangle);
-				TriQuery.ComputeResult();
-				const UE::Geometry::FIndex3i& TriangleIndices = DynamicMesh.GetTriangle(NearestTriangle);
-				const float Weight = TriQuery.TriangleBaryCoords.X * GetWeightForIndices(BonesSet, OrigIndices[TriangleIndices.A]) +
-					TriQuery.TriangleBaryCoords.Y * GetWeightForIndices(BonesSet, OrigIndices[TriangleIndices.B]) + TriQuery.TriangleBaryCoords.Z * GetWeightForIndices(BonesSet, OrigIndices[TriangleIndices.C]);
+				const float Weight = CalculateTriangleWeight(Triangle, TriangleWeights, Location);				
+				check(Weight != 0.f);
 				ConstrainedNodes.Add(SubGrid.FlatIndex(Index, true));
 				ConstrainedWeights.Add(Weight);
 			}
@@ -225,7 +277,7 @@ void FSkinnedLevelSetBuilder::AddBoneInfluence(int32 PrimaryBoneIndex, const TAr
 	Chaos::PoissonSolve<float, double, true>(ConstrainedNodes, ConstrainedWeights, SubGrid, MaxIter, Res, Weights, bCheckResidual, MinParallelBatchSize);
 	for (int32 Idx = 0; Idx < Weights.Num(); ++Idx)
 	{
-		if (Weights[Idx] > UE_SMALL_NUMBER)
+		if (Weights[Idx] > 0)
 		{
 			Chaos::TVec3<int32> SubIndex;
 			SubGrid.FlatToMultiIndex(Idx, SubIndex, true);
@@ -272,7 +324,7 @@ float FSkinnedLevelSetBuilder::GetWeightForIndices(const TSet<int32>& BoneIndice
 		if (BoneIndices.Contains(ActualBoneIndex))
 		{
 			// Use max weight for all bones.
-			Weight = FMath::Max(Weight, (float)SoftVert.InfluenceWeights[InfluenceIndex] / UE::AnimationCore::MaxRawBoneWeightFloat);
+			Weight = FMath::Max(Weight, (float)SoftVert.InfluenceWeights[InfluenceIndex] * UE::AnimationCore::InvMaxRawBoneWeightFloat);
 		}
 	}
 
diff --git a/Engine/Source/Runtime/Engine/Private/PhysicsEngine/PhysDrawing.cpp b/Engine/Source/Runtime/Engine/Private/PhysicsEngine/PhysDrawing.cpp
--- a/Engine/Source/Runtime/Engine/Private/PhysicsEngine/PhysDrawing.cpp
+++ b/Engine/Source/Runtime/Engine/Private/PhysicsEngine/PhysDrawing.cpp
@@ -41,6 +41,10 @@ static const FColor	JointLimitColor(FColor::Green);
 static const FColor	JointRefColor(FColor::Yellow);
 static const FColor JointLockedColor(255,128,10);
 
+
+static int SkinnedLatticeBoneWeight = -1;
+static FAutoConsoleVariableRef CVarClothVizDrawSkinnedLattice(TEXT("p.PhysDrawing.SkinnedLatticeBoneWeight"), SkinnedLatticeBoneWeight, TEXT("Draw skinned lattice bone weight. -1 = all lattice points"));
+
 /////////////////////////////////////////////////////////////////////////////////////
 // FKSphereElem
 /////////////////////////////////////////////////////////////////////////////////////
@@ -854,10 +858,64 @@ static void DrawSkinnedLevelSetLattice(class FPrimitiveDrawInterface* PDI, const
 	const Chaos::TUniformGrid<double, 3>& LatticeGrid = WeightedLevelSet->GetGrid();
 	const Chaos::TArrayND<Chaos::FVec3, 3>& DeformedPoints = WeightedLevelSet->GetDeformedPoints();
 	const Chaos::TArrayND<bool, 3>& EmptyCells = WeightedLevelSet->GetEmptyCells();
+	const Chaos::TArrayND<Chaos::FWeightedLatticeInfluenceData, 3>& BoneData = WeightedLevelSet->GetBoneData();
+	const int32 NumUsedBones = WeightedLevelSet->GetUsedBones().Num();
+	const int32 SkinnedLatticeBoneWeightLocal = SkinnedLatticeBoneWeight < NumUsedBones ? SkinnedLatticeBoneWeight : -1;
+
+	Chaos::TArrayND<float,3> SkinnedLatticeBoneWeights;
+	if (SkinnedLatticeBoneWeightLocal >= 0)
+	{
+		static float PointSize = 5.f;
+		SkinnedLatticeBoneWeights.SetCounts(LatticeGrid, true);
+
+		const Chaos::TVec3<int32> NodeCounts = LatticeGrid.NodeCounts();
+		for (int32 I = 0; I < NodeCounts.X; ++I)
+		{
+			for (int32 J = 0; J < NodeCounts.Y; ++J)
+			{
+				for (int32 K = 0; K < NodeCounts.Z; ++K)
+				{
+					SkinnedLatticeBoneWeights(I, J, K) = 0.f;
+					const Chaos::FWeightedLatticeInfluenceData& BoneDatum = BoneData(I, J, K);
+					for (int32 InfIdx = 0; InfIdx < BoneDatum.NumInfluences; ++InfIdx)
+					{
+						if (BoneDatum.BoneIndices[InfIdx] == SkinnedLatticeBoneWeightLocal)
+						{
+							const float Weight = BoneDatum.BoneWeights[InfIdx];
+							SkinnedLatticeBoneWeights(I, J, K) = Weight;
+							PDI->DrawPoint(LocalToWorld.TransformPosition(FVector(DeformedPoints(I, J, K))), FLinearColor(Weight, Weight, Weight), PointSize, SDPG_World);
+							break;
+						}
+					}
+				}
+			}
+		}
+	}
 
+	auto ShouldDrawCell = [&SkinnedLatticeBoneWeights, &EmptyCells](int32 I, int32 J, int32 K)
+	{
+		if (EmptyCells(I, J, K))
+		{
+			return false;
+		}
+		if (SkinnedLatticeBoneWeights.Num() == 0)
+		{
+			return true;
+		}
+		if (SkinnedLatticeBoneWeights(I, J, K) == 0.f ||
+			SkinnedLatticeBoneWeights(I, J, K + 1) == 0.f ||
+			SkinnedLatticeBoneWeights(I, J + 1, K) == 0.f ||
+			SkinnedLatticeBoneWeights(I, J + 1, K + 1) == 0.f ||
+			SkinnedLatticeBoneWeights(I + 1, J, K) == 0.f ||
+			SkinnedLatticeBoneWeights(I + 1, J, K + 1) == 0.f ||
+			SkinnedLatticeBoneWeights(I + 1, J + 1, K) == 0.f ||
+			SkinnedLatticeBoneWeights(I + 1, J + 1, K + 1) == 0.f)
+		{
+			return false;
+		}
+		return true;
+	};
 
-	// TODO: downres drawing lattice when res is too high. 
-	// Does two passes to count lines and reserve help?
 	const Chaos::TVec3<int32> CellCounts = LatticeGrid.Counts();
 	for (int32 I = 0; I < CellCounts.X; ++I)
 	{
@@ -865,7 +923,7 @@ static void DrawSkinnedLevelSetLattice(class FPrimitiveDrawInterface* PDI, const
 		{
 			for (int32 K = 0; K < CellCounts.Z; ++K)
 			{
-				if (!EmptyCells(I, J, K))
+				if (ShouldDrawCell(I, J, K))
 				{
 					const FVector P000 = LocalToWorld.TransformPosition(FVector(DeformedPoints(I, J, K)));
 					const FVector P001 = LocalToWorld.TransformPosition(FVector(DeformedPoints(I, J, K + 1)));
@@ -889,7 +947,6 @@ static void DrawSkinnedLevelSetLattice(class FPrimitiveDrawInterface* PDI, const
 					PDI->DrawLine(P100, P110, Color, SDPG_World);
 					PDI->DrawLine(P101, P111, Color, SDPG_World);
 					PDI->DrawLine(P110, P111, Color, SDPG_World);
-
 				}
 			}
 		}
diff --git a/Engine/Source/Runtime/Experimental/Chaos/Private/Chaos/WeightedLatticeImplicitObject.cpp b/Engine/Source/Runtime/Experimental/Chaos/Private/Chaos/WeightedLatticeImplicitObject.cpp
--- a/Engine/Source/Runtime/Experimental/Chaos/Private/Chaos/WeightedLatticeImplicitObject.cpp
+++ b/Engine/Source/Runtime/Experimental/Chaos/Private/Chaos/WeightedLatticeImplicitObject.cpp
@@ -781,7 +781,9 @@ void FWeightedLatticeImplicitObjectBuilder::NormalizeBoneWeights()
 		{
 			TotalWeight += Data.BoneWeights[InfIdx];
 		}
-		if (TotalWeight < UE_SMALL_NUMBER)
+		constexpr float MIN_TOTAL_OUTER_WEIGHT = 0.1f;
+		if(TotalWeight < UE_SMALL_NUMBER ||
+			(BuildData[Idx].WeightsAreOuter && TotalWeight < MIN_TOTAL_OUTER_WEIGHT))
 		{
 			Data.NumInfluences = 0;
 		}
