Add PhysicsAssetUtilsModule boilerplate.

This fixes a crash when undoing creating SkinnedLevelSets because PhysicsAssetUtils wasn't using the correct global allocator.

diff --git a/Engine/Source/Developer/PhysicsUtilities/Private/PhysicsAssetUtils.cpp b/Engine/Source/Developer/PhysicsUtilities/Private/PhysicsAssetUtils.cpp
--- a/Engine/Source/Developer/PhysicsUtilities/Private/PhysicsAssetUtils.cpp
+++ b/Engine/Source/Developer/PhysicsUtilities/Private/PhysicsAssetUtils.cpp
@@ -1294,3 +1294,5 @@ void SanitizeRestrictedContent(UPhysicsAsset* PhysAsset)
 }
 
 }; // namespace FPhysicsAssetUtils
+
+IMPLEMENT_MODULE(FDefaultModuleImpl, PhysicsAssetUtils)
