[UE] Phat: Isolate and Mirror can misbehave when used together
- Added support for re-creating physics bodies without changing their index in the owning physics asset.
- Replaced calls to createNewBody with RecreateBody in mirroring code. This means that all references to a mirrored body via index remain valid after mirroring.
- Reduced code duplication in Mirror function

diff --git a/Engine/Source/Developer/PhysicsUtilities/Private/PhysicsAssetUtils.cpp b/Engine/Source/Developer/PhysicsUtilities/Private/PhysicsAssetUtils.cpp
--- a/Engine/Source/Developer/PhysicsUtilities/Private/PhysicsAssetUtils.cpp
+++ b/Engine/Source/Developer/PhysicsUtilities/Private/PhysicsAssetUtils.cpp
@@ -32,11 +32,96 @@
 
 namespace FPhysicsAssetUtils
 {
-	static const float	DefaultPrimSize = 15.0f;
-	static const float	MinPrimSize = 0.5f;
 
-	// Forward declarations
-	bool CreateCollisionFromBoneInternal(UBodySetup* bs, USkeletalMesh* skelMesh, int32 BoneIndex, const FPhysAssetCreateParams& Params, const FBoneVertInfo& Info, const FSkinnedBoneTriangleCache& TriangleCache);
+static const float	DefaultPrimSize = 15.0f;
+static const float	MinPrimSize = 0.5f;
+
+// Forward declarations
+bool CreateCollisionFromBoneInternal(UBodySetup* bs, USkeletalMesh* skelMesh, int32 BoneIndex, const FPhysAssetCreateParams& Params, const FBoneVertInfo& Info, const FSkinnedBoneTriangleCache& TriangleCache);
+
+// File Scope Utility Functions //
+void RemoveBodyFromCollisionTableInternal(UPhysicsAsset* const PhysAsset, const int32 BodyIndex, const bool bDecrementFollowingIndices)
+{
+	check(PhysAsset);
+
+	if (PhysAsset)
+	{
+		// First we must correct the CollisionDisableTable.
+		// All elements which refer to bodyIndex are removed.
+		// All elements which refer to a body with index >bodyIndex are adjusted. 
+
+		TMap<FRigidBodyIndexPair, bool> NewCDT;
+		for (int32 i = 1; i < PhysAsset->SkeletalBodySetups.Num(); i++)
+		{
+			for (int32 j = 0; j < i; j++)
+			{
+				FRigidBodyIndexPair Key(j, i);
+
+				// If there was an entry for this pair, and it doesn't refer to the removed body, we need to add it to the new CDT.
+				if (PhysAsset->CollisionDisableTable.Find(Key))
+				{
+					if (i != BodyIndex && j != BodyIndex)
+					{
+						int32 NewI = i;
+						int32 NewJ = j;
+
+						if (bDecrementFollowingIndices)
+						{
+							// Decrement index of all bodies with index greater than the supplied BodyIndex. This is useful when the BodyIndex is going to be removed from the physics asset.
+							if (NewI > BodyIndex) { --NewI; }
+							if (NewJ > BodyIndex) { --NewJ; }
+						}
+
+						FRigidBodyIndexPair NewKey(NewJ, NewI);
+						NewCDT.Add(NewKey, 0);
+					}
+				}
+			}
+		}
+
+		PhysAsset->CollisionDisableTable = NewCDT;
+	}
+}
+
+void RemoveBodyConstraintsInternal(UPhysicsAsset* const PhysAsset, const int32 BodyIndex)
+{
+	// Now remove any constraints that were attached to this body.
+	// This is a bit yuck and slow...
+	TArray<int32> Constraints;
+	PhysAsset->BodyFindConstraints(BodyIndex, Constraints);
+
+	while (Constraints.Num() > 0)
+	{
+		DestroyConstraint(PhysAsset, Constraints[0]);
+		PhysAsset->BodyFindConstraints(BodyIndex, Constraints);
+	}
+}
+
+void CreateNewBodyAtIndexInternal(UPhysicsAsset* const PhysAsset, const FName InBodyName, const FPhysAssetCreateParams& Params, const int32 BodyIndex)
+{
+	USkeletalBodySetup* NewBodySetup = NewObject<USkeletalBodySetup>(PhysAsset, NAME_None, RF_Transactional);
+	// make default to be use complex as simple 
+	NewBodySetup->CollisionTraceFlag = CTF_UseSimpleAsComplex;
+	// newly created bodies default to simulating
+	NewBodySetup->PhysicsType = PhysType_Default;
+
+	NewBodySetup->BoneName = InBodyName;
+
+	PhysAsset->SkeletalBodySetups[BodyIndex] = NewBodySetup;
+
+	PhysAsset->UpdateBodySetupIndexMap();
+	PhysAsset->UpdateBoundsBodiesArray();
+
+	if (Params.bDisableCollisionsByDefault)
+	{
+		for (int i = 0; i < PhysAsset->SkeletalBodySetups.Num(); ++i)
+		{
+			PhysAsset->DisableCollision(i, BodyIndex);
+		}
+	}
+}
+
+// Externally Accessible Functions //
 
 /** Returns INDEX_NONE if no children in the visual asset or if more than one parent */
 static int32 GetChildIndex(int32 BoneIndex, USkeletalMesh* SkelMesh, const TArray<FBoneVertInfo>& Infos)
@@ -1379,8 +1464,7 @@ void DestroyConstraint(UPhysicsAsset* PhysAsset, int32 ConstraintIndex)
 	PhysAsset->ConstraintSetup.RemoveAt(ConstraintIndex);
 }
 
-
-int32 CreateNewBody(UPhysicsAsset* PhysAsset, FName InBodyName, const FPhysAssetCreateParams& Params)
+int32 CreateNewBody(UPhysicsAsset* const PhysAsset, const FName InBodyName, const FPhysAssetCreateParams& Params)
 {
 	check(PhysAsset);
 
@@ -1390,72 +1474,25 @@ int32 CreateNewBody(UPhysicsAsset* PhysAsset, FName InBodyName, const FPhysAsset
 		return BodyIndex; // if we already have one for this name - just return that.
 	}
 
-	USkeletalBodySetup* NewBodySetup = NewObject<USkeletalBodySetup>(PhysAsset, NAME_None, RF_Transactional);
-	// make default to be use complex as simple 
-	NewBodySetup->CollisionTraceFlag = CTF_UseSimpleAsComplex;
-	// newly created bodies default to simulating
-	NewBodySetup->PhysicsType = PhysType_Default;
-
-	int32 BodySetupIndex = PhysAsset->SkeletalBodySetups.Add( NewBodySetup );
-	NewBodySetup->BoneName = InBodyName;
-
-	PhysAsset->UpdateBodySetupIndexMap();
-	PhysAsset->UpdateBoundsBodiesArray();
+	BodyIndex = PhysAsset->SkeletalBodySetups.AddDefaulted();
+	CreateNewBodyAtIndexInternal(PhysAsset, InBodyName, Params, BodyIndex);
 
-	if (Params.bDisableCollisionsByDefault)
-	{
-		for (int i = 0; i < PhysAsset->SkeletalBodySetups.Num(); ++i)
-		{
-			PhysAsset->DisableCollision(i, BodySetupIndex);
-		}
-	}
-	// Return index of new body.
-	return BodySetupIndex;
+	return BodyIndex; // Return index of new body.
 }
 
-void DestroyBody(UPhysicsAsset* PhysAsset, int32 bodyIndex)
+void RecreateBody(UPhysicsAsset* const PhysAsset, const FName InBodyName, const FPhysAssetCreateParams& Params, const int32 BodyIndex)
 {
-	check(PhysAsset);
-
-	// First we must correct the CollisionDisableTable.
-	// All elements which refer to bodyIndex are removed.
-	// All elements which refer to a body with index >bodyIndex are adjusted. 
-
-	TMap<FRigidBodyIndexPair,bool> NewCDT;
-	for(int32 i=1; i<PhysAsset->SkeletalBodySetups.Num(); i++)
-	{
-		for(int32 j=0; j<i; j++)
-		{
-			FRigidBodyIndexPair Key(j,i);
-
-			// If there was an entry for this pair, and it doesn't refer to the removed body, we need to add it to the new CDT.
-			if( PhysAsset->CollisionDisableTable.Find(Key) )
-			{
-				if(i != bodyIndex && j != bodyIndex)
-				{
-					int32 NewI = (i > bodyIndex) ? i-1 : i;
-					int32 NewJ = (j > bodyIndex) ? j-1 : j;
-
-					FRigidBodyIndexPair NewKey(NewJ, NewI);
-					NewCDT.Add(NewKey, 0);
-				}
-			}
-		}
-	}
-
-	PhysAsset->CollisionDisableTable = NewCDT;
-
-	// Now remove any constraints that were attached to this body.
-	// This is a bit yuck and slow...
-	TArray<int32> Constraints;
-	PhysAsset->BodyFindConstraints(bodyIndex, Constraints);
-
-	while(Constraints.Num() > 0)
-	{
-		DestroyConstraint( PhysAsset, Constraints[0] );
-		PhysAsset->BodyFindConstraints(bodyIndex, Constraints);
-	}
+	RemoveBodyFromCollisionTableInternal(PhysAsset, BodyIndex, false);
+	RemoveBodyConstraintsInternal(PhysAsset, BodyIndex);
+	PhysAsset->SkeletalBodySetups[BodyIndex] = TObjectPtr<USkeletalBodySetup>(); // Reset TObjectPtr to body setup that will be replaced.
+	CreateNewBodyAtIndexInternal(PhysAsset, InBodyName, Params, BodyIndex);
+}
 
+void DestroyBody(UPhysicsAsset* const PhysAsset, const int32 bodyIndex)
+{
+	RemoveBodyFromCollisionTableInternal(PhysAsset, bodyIndex, true);
+	RemoveBodyConstraintsInternal(PhysAsset, bodyIndex);
+	
 	// Remove pointer from array. Actual objects will be garbage collected.
 	PhysAsset->SkeletalBodySetups.RemoveAt(bodyIndex);
 
diff --git a/Engine/Source/Developer/PhysicsUtilities/Public/PhysicsAssetUtils.h b/Engine/Source/Developer/PhysicsUtilities/Public/PhysicsAssetUtils.h
--- a/Engine/Source/Developer/PhysicsUtilities/Public/PhysicsAssetUtils.h
+++ b/Engine/Source/Developer/PhysicsUtilities/Public/PhysicsAssetUtils.h
@@ -219,7 +219,7 @@ namespace FPhysicsAssetUtils
 	 * @param	InBodyName			Name of the new body
 	 * @return	The Index of the newly created body.
 	 */
-	PHYSICSUTILITIES_API int32 CreateNewBody(UPhysicsAsset* PhysAsset, FName InBodyName, const FPhysAssetCreateParams& Params);
+	PHYSICSUTILITIES_API int32 CreateNewBody(UPhysicsAsset* const PhysAsset, const FName InBodyName, const FPhysAssetCreateParams& Params);
 	
 	/** 
 	 * Destroys the specified body
@@ -227,7 +227,17 @@ namespace FPhysicsAssetUtils
 	 * @param	PhysAsset			The PhysicsAsset for which the body should be destroyed
 	 * @param	BodyIndex			Index of the body to destroy
 	 */
-	PHYSICSUTILITIES_API void DestroyBody(UPhysicsAsset* PhysAsset, int32 BodyIndex);
+	PHYSICSUTILITIES_API void DestroyBody(UPhysicsAsset* const PhysAsset, const int32 BodyIndex);
+
+	/**
+	 * Initialize a new BodySetup at a specified index in the physics asset and create a default BodyInstance if there is not one for this body already.
+	 *
+	 * @param	PhysAsset			The PhysicsAsset to create the body for
+	 * @param	InBodyName			Name of the new body
+	 * @param	Params				Parameters used when configuring the body
+	 * @param	BodyIndex			Index of the body to initialize in the Physics Asset
+	 */
+	PHYSICSUTILITIES_API void RecreateBody(UPhysicsAsset* const PhysAsset, const FName InBodyName, const FPhysAssetCreateParams& Params, const int32 BodyIndex);
 
 	/**
 	* Whether or not Constraints are allowed to be created (due to asset list filtering)
diff --git a/Engine/Source/Editor/PhysicsAssetEditor/Private/PhysicsAssetEditorSharedData.cpp b/Engine/Source/Editor/PhysicsAssetEditor/Private/PhysicsAssetEditorSharedData.cpp
--- a/Engine/Source/Editor/PhysicsAssetEditor/Private/PhysicsAssetEditorSharedData.cpp
+++ b/Engine/Source/Editor/PhysicsAssetEditor/Private/PhysicsAssetEditorSharedData.cpp
@@ -1,6 +1,7 @@
 // Copyright Epic Games, Inc. All Rights Reserved.
 
 #include "PhysicsAssetEditorSharedData.h"
+#include "Animation/MirrorDataTable.h"
 #include "PhysicsAssetEditorPhysicsHandleComponent.h"
 #include "PhysicsAssetRenderUtils.h"
 #include "PhysicsEngine/RigidBodyIndexPair.h"
@@ -566,6 +567,28 @@ struct FMirrorInfo
 	}
 };
 
+template<typename GeometryElementType> void MirrorPrimitives(TArray< GeometryElementType >& PrimitiveCollection)
+{
+	static const FQuat ArtistMirrorConvention(1, 0, 0, 0);	//used to be (0 0 1 0)
+															// how Epic Maya artists rig the right and left orientation differently.  todo: perhaps move to cvar W
+	
+	for (GeometryElementType& Primitive : PrimitiveCollection)
+	{
+		Primitive.Rotation = (Primitive.Rotation.Quaternion() * ArtistMirrorConvention).Rotator();
+		Primitive.Center = -Primitive.Center;
+		Primitive.SetName(UMirrorDataTable::GetSettingsMirrorName(Primitive.GetName()));
+	}
+}
+
+template<> void MirrorPrimitives<FKSphereElem>(TArray< FKSphereElem >& PrimitiveCollection)
+{
+	for (FKSphereElem& Primitive : PrimitiveCollection)
+	{
+		Primitive.Center = -Primitive.Center;
+		Primitive.SetName(UMirrorDataTable::GetSettingsMirrorName(Primitive.GetName()));
+	}
+}
+
 void FPhysicsAssetEditorSharedData::Mirror()
 {
 	USkeletalMesh* EditorSkelMesh = PhysicsAsset->GetPreviewMesh();
@@ -615,37 +638,19 @@ void FPhysicsAssetEditorSharedData::Mirror()
 			{
 				UBodySetup * SrcBody = PhysicsAsset->SkeletalBodySetups[MirrorInfo.BodyIndex];
 				const FScopedTransaction Transaction(NSLOCTEXT("PhysicsAssetEditor", "MirrorBody", "MirrorBody"));
-				MakeNewBody(MirrorBoneIndex, false);
+				RecreateBody(MirrorBoneIndex, false); // Create a new body at the same index as the one being replaced. This ensures that all references to this body via index remain valid.
 
-				int32 MirrorBodyIndex = PhysicsAsset->FindControllingBodyIndex(EditorSkelMesh, MirrorBoneIndex);
-
-				UBodySetup * DestBody = PhysicsAsset->SkeletalBodySetups[MirrorBodyIndex];
+				const int32 MirrorBodyIndex = PhysicsAsset->FindControllingBodyIndex(EditorSkelMesh, MirrorBoneIndex);
+				UBodySetup* const DestBody = PhysicsAsset->SkeletalBodySetups[MirrorBodyIndex];
 				DestBody->Modify();
 				DestBody->CopyBodyPropertiesFrom(SrcBody);
-				
-				FQuat ArtistMirrorConvention(1,0,0,0);   //used to be (0 0 1 0)
-														 // how Epic Maya artists rig the right and left orientation differently.  todo: perhaps move to cvar W
 
-				for (FKSphylElem& Sphyl : DestBody->AggGeom.SphylElems)
-				{
-					Sphyl.Rotation	= (Sphyl.Rotation.Quaternion()*ArtistMirrorConvention).Rotator();
-					Sphyl.Center = -Sphyl.Center;
-				}
-				for (FKBoxElem& Box : DestBody->AggGeom.BoxElems)
-				{
-					Box.Rotation	= (Box.Rotation.Quaternion()*ArtistMirrorConvention).Rotator();
-					Box.Center		= -Box.Center;
-				}
-				for (FKSphereElem& Sphere : DestBody->AggGeom.SphereElems)
-				{
-					Sphere.Center = -Sphere.Center;
-				}
-				for (FKTaperedCapsuleElem& TaperedCapsule : DestBody->AggGeom.TaperedCapsuleElems)
-				{
-					TaperedCapsule.Rotation	= (TaperedCapsule.Rotation.Quaternion()*ArtistMirrorConvention).Rotator();
-					TaperedCapsule.Center	= -TaperedCapsule.Center;
-				}
-				int32 MirrorConstraintIndex = PhysicsAsset->FindConstraintIndex(DestBody->BoneName);
+				MirrorPrimitives(DestBody->AggGeom.SphylElems);
+				MirrorPrimitives(DestBody->AggGeom.BoxElems);
+				MirrorPrimitives(DestBody->AggGeom.SphereElems);
+				MirrorPrimitives(DestBody->AggGeom.TaperedCapsuleElems);
+
+				const int32 MirrorConstraintIndex = PhysicsAsset->FindConstraintIndex(DestBody->BoneName);
 				if(PhysicsAsset->ConstraintSetup.IsValidIndex(MirrorConstraintIndex) && PhysicsAsset->ConstraintSetup.IsValidIndex(MirrorInfo.ConstraintIndex))
 				{
 					UPhysicsConstraintTemplate * FromConstraint = PhysicsAsset->ConstraintSetup[MirrorInfo.ConstraintIndex];
@@ -2456,6 +2461,36 @@ void FPhysicsAssetEditorSharedData::InitConstraintSetup(UPhysicsConstraintTempla
 	SetCollisionBetween(ChildBodyIndex, ParentBodyIndex, false);
 }
 
+void FPhysicsAssetEditorSharedData::RecreateBody(const int32 NewBoneIndex, const bool bAutoSelect)
+{
+	RecreateBody(GetDefault<UPhysicsAssetGenerationSettings>()->CreateParams, NewBoneIndex, bAutoSelect);
+}
+
+void FPhysicsAssetEditorSharedData::RecreateBody(const FPhysAssetCreateParams& BodyData, const int32 BoneIndex, const bool bAutoSelect)
+{
+	USkeletalMesh* EditorSkelMesh = PhysicsAsset->GetPreviewMesh();
+	if (EditorSkelMesh == nullptr)
+	{
+		return;
+	}
+	
+	PhysicsAsset->Modify();
+
+	const FName BoneName = EditorSkelMesh->GetRefSkeleton().GetBoneName(BoneIndex);
+	const int32 BodyIndex = PhysicsAsset->FindBodyIndex(BoneName);
+
+	FPhysicsAssetUtils::RecreateBody(PhysicsAsset, BoneName, BodyData, BodyIndex); // Create a new physics body setup at the same index as the original body setup.
+
+	BroadcastHierarchyChanged();
+
+	if (bAutoSelect)
+	{
+		SetSelectedBodyAnyPrimitive(BodyIndex, true);
+	}
+
+	RefreshPhysicsAssetChange(PhysicsAsset);
+}
+
 void FPhysicsAssetEditorSharedData::MakeNewBody(int32 NewBoneIndex, bool bAutoSelect)
 {
 	MakeNewBody(GetDefault<UPhysicsAssetGenerationSettings>()->CreateParams, NewBoneIndex, bAutoSelect);
diff --git a/Engine/Source/Editor/PhysicsAssetEditor/Private/PhysicsAssetEditorSharedData.h b/Engine/Source/Editor/PhysicsAssetEditor/Private/PhysicsAssetEditorSharedData.h
--- a/Engine/Source/Editor/PhysicsAssetEditor/Private/PhysicsAssetEditorSharedData.h
+++ b/Engine/Source/Editor/PhysicsAssetEditor/Private/PhysicsAssetEditorSharedData.h
@@ -228,6 +228,8 @@ public:
 	void RefreshPhysicsAssetChange(const UPhysicsAsset* InPhysAsset, bool bFullClothRefresh = true);
 	void MakeNewBody(int32 NewBoneIndex, bool bAutoSelect = true);
 	void MakeNewBody(const FPhysAssetCreateParams& NewBodyData, const int32 NewBoneIndex, const bool bAutoSelect = true);
+	void RecreateBody(const int32 BoneIndex, const bool bAutoSelect = true); // Create a new physics body to replace an existing one.
+	void RecreateBody(const FPhysAssetCreateParams& BodyData, const int32 BoneIndex, const bool bAutoSelect = true); // Create a new physics body to replace an existing one.
 	void MakeNewConstraints(int32 ParentBodyIndex, const TArray<int32>& ChildBodyIndices);
 	void MakeNewConstraint(int32 ParentBodyIndex, int32 ChildBodyIndex);
 	void CopySelectedBodiesAndConstraintsToClipboard(int32& OutNumCopiedBodies, int32& OutNumCopiedConstraints, int32& OutNumCopiedDisabledCollisionPairs);
