Use FMemory::Malloc instead of GMalloc in ConvexDecompTool.

GMalloc is the low-level allocator. It is not transactionally safe and doesn't interact with LLM stats. Higher level code should prefer FMemory::Malloc instead.

diff --git a/Engine/Source/Developer/PhysicsUtilities/Private/ConvexDecompTool.cpp b/Engine/Source/Developer/PhysicsUtilities/Private/ConvexDecompTool.cpp
--- a/Engine/Source/Developer/PhysicsUtilities/Private/ConvexDecompTool.cpp
+++ b/Engine/Source/Developer/PhysicsUtilities/Private/ConvexDecompTool.cpp
@@ -12,6 +12,7 @@
 #include "ThirdParty/VHACD/public/VHACD.h"
 #include "ThirdParty/VHACD/inc/btAlignedAllocator.h"
 
+#include "HAL/UnrealMemory.h"
 #include "PhysicsEngine/BodySetup.h"
 #include "Async/Async.h"
 
@@ -104,22 +105,22 @@ static FVHACDTaskRunner VHACDTaskRunner;
 
 static void* btAlignedAllocImpl(size_t size, int32_t alignment)
 {
-	return GMalloc->Malloc(size, alignment);
+	return FMemory::Malloc(size, alignment);
 }
 
 static void btAlignedFreeImpl(void* memblock)
 {
-	GMalloc->Free(memblock);
+	FMemory::Free(memblock);
 }
 
 static void* btAllocImpl(size_t size)
 {
-	return GMalloc->Malloc(size);
+	return FMemory::Malloc(size);
 }
 
 static void btFreeImpl(void* memblock)
 {
-	GMalloc->Free(memblock);
+	FMemory::Free(memblock);
 }
 
 static void InitParameters(IVHACD::Parameters &VHACD_Params, uint32 InHullCount, uint32 InMaxHullVerts,uint32 InResolution)
