Expose functions for calculating covariance matrix of a set of points - and retrieving the dominant eigenvector. Updated implementations to take any vector type

diff --git a/Engine/Source/Developer/PhysicsUtilities/Private/PhysicsAssetUtils.cpp b/Engine/Source/Developer/PhysicsUtilities/Private/PhysicsAssetUtils.cpp
--- a/Engine/Source/Developer/PhysicsUtilities/Private/PhysicsAssetUtils.cpp
+++ b/Engine/Source/Developer/PhysicsUtilities/Private/PhysicsAssetUtils.cpp
@@ -621,56 +621,63 @@ bool CreateFromSkeletalMesh(UPhysicsAsset* PhysicsAsset, USkeletalMesh* SkelMesh
 	return bSuccess;
 }
 
-FMatrix ComputeCovarianceMatrix(const FBoneVertInfo& VertInfo)
+template<typename RealType>
+FMatrix ComputeViewCovarianceMatrix(TArrayView<const UE::Math::TVector<RealType>, int> PointView)
 {
-	if (VertInfo.Positions.Num() == 0)
+	using FVectorType = UE::Math::TVector<RealType>;
+
+	if(PointView.Num() == 0)
 	{
 		return FMatrix::Identity;
 	}
 
-	const TArray<FVector3f> & Positions = VertInfo.Positions;
-
 	//get average
-	const float N = Positions.Num();
-	FVector U = FVector::ZeroVector;
-	for (int32 i = 0; i < N; ++i)
+	const RealType N = PointView.Num();
+	FVectorType U = FVectorType::ZeroVector;
+	
+	for(const FVectorType& Point : PointView)
 	{
-		U += (FVector)Positions[i];
+		U += Point;
 	}
-
+	
 	U = U / N;
 
 	//compute error terms
-	TArray<FVector> Errors;
+	TArray<FVectorType> Errors;
 	Errors.AddUninitialized(N);
 
-	for (int32 i = 0; i < N; ++i)
+	for(int32 i = 0; i < N; ++i)
 	{
-		Errors[i] = (FVector)Positions[i] - U;
+		Errors[i] = PointView[i] - U;
 	}
 
 	FMatrix Covariance = FMatrix::Identity;
-	for (int32 j = 0; j < 3; ++j)
+	for(int32 j = 0; j < 3; ++j)
 	{
-		FVector Axis = FVector::ZeroVector;
-		FVector::FReal* Cj = &Axis.X;
-		for (int32 k = 0; k < 3; ++k)
+		FVectorType Axis = FVectorType::ZeroVector;
+		RealType* Cj = &Axis.X;
+		for(int32 k = 0; k < 3; ++k)
 		{
-			float Cjk = 0.f;
-			for (int32 i = 0; i < N; ++i)
+			RealType Cjk = 0.f;
+			for(int32 i = 0; i < N; ++i)
 			{
-				const FVector::FReal* error = &Errors[i].X;
-				Cj[k] += error[j] * error[k];
+				const RealType* Error = &Errors[i].X;
+				Cj[k] += Error[j] * Error[k];
 			}
 			Cj[k] /= N;
 		}
 
-		Covariance.SetAxis(j, Axis);
+		Covariance.SetAxis(j, UE::Math::TVector<double>(Axis));
 	}
 
 	return Covariance;
 }
 
+FMatrix ComputeCovarianceMatrix(const FBoneVertInfo& VertInfo)
+{
+	return ComputeViewCovarianceMatrix(MakeArrayView(VertInfo.Positions));
+}
+
 FVector ComputeEigenVector(const FMatrix& A)
 {
 	//using the power method: this is ok because we only need the dominate eigenvector and speed is not critical: http://en.wikipedia.org/wiki/Power_iteration
@@ -1518,6 +1525,11 @@ void SanitizeRestrictedContent(UPhysicsAsset* PhysAsset)
 	check(PhysAsset);
 }
 
+// Explicitly instantiate and export covariance matrix calc
+template PHYSICSUTILITIES_API FMatrix ComputeViewCovarianceMatrix(TArrayView<const UE::Math::TVector<float>, int> PointView);
+template PHYSICSUTILITIES_API FMatrix ComputeViewCovarianceMatrix(TArrayView<const UE::Math::TVector<double>, int> PointView);
+
 }; // namespace FPhysicsAssetUtils
 
 IMPLEMENT_MODULE(FDefaultModuleImpl, PhysicsUtilities)
+
diff --git a/Engine/Source/Developer/PhysicsUtilities/Public/PhysicsAssetUtils.h b/Engine/Source/Developer/PhysicsUtilities/Public/PhysicsAssetUtils.h
--- a/Engine/Source/Developer/PhysicsUtilities/Public/PhysicsAssetUtils.h
+++ b/Engine/Source/Developer/PhysicsUtilities/Public/PhysicsAssetUtils.h
@@ -250,4 +250,25 @@ namespace FPhysicsAssetUtils
 	PHYSICSUTILITIES_API bool CanCreateConstraints();
 
 	PHYSICSUTILITIES_API void SanitizeRestrictedContent(UPhysicsAsset* PhysAsset);
+
+	/**
+	 * Compute a covariance matrix describing the distribution of points in the provided set
+	 * 
+	 * @param PointView An array view of a vector type to calculate the covariance matrix for
+	 * @return Calculated covariance matrix
+	 */
+	template<typename RealType>
+	FMatrix ComputeViewCovarianceMatrix(TArrayView<const UE::Math::TVector<RealType>, int> PointView);
+	PHYSICSUTILITIES_API FMatrix ComputeCovarianceMatrix(const FBoneVertInfo& VertInfo);
+
+	extern template FMatrix ComputeViewCovarianceMatrix(TArrayView<const UE::Math::TVector<float>, int> PointView);
+	extern template FMatrix ComputeViewCovarianceMatrix(TArrayView<const UE::Math::TVector<double>, int> PointView);
+
+	/**
+	 * Compute the dominant eigenvector of a provided matrix
+	 * 
+	 * @param InMatrix Matrix input
+	 * @return Domininant eigenvector of InMatrix
+	 */
+	PHYSICSUTILITIES_API FVector ComputeEigenVector(const FMatrix& InMatrix);
 };
\ No newline at end of file
