PR #13556: Fix spelling error in ConvexDecompTool.cpp

diff --git a/Engine/Source/Developer/PhysicsUtilities/Private/ConvexDecompTool.cpp b/Engine/Source/Developer/PhysicsUtilities/Private/ConvexDecompTool.cpp
--- a/Engine/Source/Developer/PhysicsUtilities/Private/ConvexDecompTool.cpp
+++ b/Engine/Source/Developer/PhysicsUtilities/Private/ConvexDecompTool.cpp
@@ -142,7 +142,7 @@ static void InitParameters(IVHACD::Parameters &VHACD_Params, uint32 InHullCount,
 	VHACD_Params.m_maxConvexHulls = InHullCount;	// The number of convex hulls requested by the artists/designer
 	VHACD_Params.m_oclAcceleration = false;
 	VHACD_Params.m_minVolumePerCH = 0.003f; // this should be around 1 / (3 * m_resolution ^ (1/3))
-	VHACD_Params.m_projectHullVertices = true; // Project the approximate hull vertices onto the original source mesh and use highest precision results opssible.
+	VHACD_Params.m_projectHullVertices = true; // Project the approximate hull vertices onto the original source mesh and use highest precision results possible.
 	VHACD_Params.m_pca = 1; // align mesh for more optimal processing of long diagonal meshes
 }
 
