Deprecated FPhysicsAssetUtils::SanitizeRestrictedContent since it no longer does anything.

#Synced-CL 48038701

diff --git a/Engine/Source/Developer/PhysicsUtilities/Public/PhysicsAssetUtils.h b/Engine/Source/Developer/PhysicsUtilities/Public/PhysicsAssetUtils.h
--- a/Engine/Source/Developer/PhysicsUtilities/Public/PhysicsAssetUtils.h
+++ b/Engine/Source/Developer/PhysicsUtilities/Public/PhysicsAssetUtils.h
@@ -249,7 +249,7 @@ namespace FPhysicsAssetUtils
 	*/
 	PHYSICSUTILITIES_API bool CanCreateConstraints();
 
-	PHYSICSUTILITIES_API void SanitizeRestrictedContent(UPhysicsAsset* PhysAsset);
+	UE_DEPRECATED(5.8, "SanitizeRestrictedContent is no longer used.") PHYSICSUTILITIES_API void SanitizeRestrictedContent(UPhysicsAsset* PhysAsset);
 
 	/**
 	 * Compute a covariance matrix describing the distribution of points in the provided set
