PhAT: Make convex (and primitive) generation work with meshes when there are helper bones. Generally prevents shape generation from deleting bodies (e.g. if their shapes can't be calculated/updated). Adds option to create convex hulls for shapes.

#Synced-CL 50944182

diff --git a/Engine/Source/Developer/PhysicsUtilities/Private/PhysicsAssetUtils.cpp b/Engine/Source/Developer/PhysicsUtilities/Private/PhysicsAssetUtils.cpp
--- a/Engine/Source/Developer/PhysicsUtilities/Private/PhysicsAssetUtils.cpp
+++ b/Engine/Source/Developer/PhysicsUtilities/Private/PhysicsAssetUtils.cpp
@@ -38,8 +38,15 @@ namespace FPhysicsAssetUtils
 static const float	DefaultPrimSize = 15.0f;
 static const float	MinPrimSize = 0.5f;
 
-// Forward declarations
-bool CreateCollisionFromBoneInternal(UBodySetup* bs, USkeletalMesh* skelMesh, int32 BoneIndex, const FPhysAssetCreateParams& Params, const FBoneVertInfo& Info, const FSkinnedBoneTriangleCache& TriangleCache);
+// This updates the collision in BodySetup associated with BoneIndex from SkelMesh. Often this is
+// being done as part of a "bulk" operation, so the caller can pass in a physics asset and a list of
+// bones that are going to be made physical - these can be used to help identify which bones count
+// as being tied to BoneIndex, and thus which vertices etc from the skin to use when initialising
+// the collision shapes.
+bool CreateCollisionFromBoneInternal(
+	UBodySetup* BodySetup, USkeletalMesh* SkelMesh, int32 BoneIndex, const FPhysAssetCreateParams& Params,
+	const FBoneVertInfo& Info, const FSkinnedBoneTriangleCache& TriangleCache, const UPhysicsAsset* PhysicsAsset,
+	TConstArrayView<int32> BoneIndicesInCreation);
 
 // File Scope Utility Functions //
 void RemoveBodyFromCollisionTableInternal(UPhysicsAsset* const PhysAsset, const int32 BodyIndex, const bool bDecrementFollowingIndices)
@@ -436,6 +443,18 @@ bool CreateFromSkeletalMeshInternal(UPhysicsAsset* PhysicsAsset, USkeletalMesh*
 		SlowTask.MakeDialog();
 	}
 
+	// When creating bodies, sometimes we will need to find other non-physical (helper) bones to
+	// include. For this, we need to know which physical bones to exclude.
+	TArray<int32> PhysicalBoneIndices;
+	for (int32 BoneIndex = 0; BoneIndex < NumBones; BoneIndex++)
+	{
+		// Determine if we should create a physics body for this bone
+		if (ShouldMakeBone(BoneIndex))
+		{
+			PhysicalBoneIndices.Emplace(BoneIndex);
+		}
+	}
+
 	for (int32 BoneIndex = 0; BoneIndex < NumBones; BoneIndex++)
 	{
 		// Determine if we should create a physics body for this bone
@@ -448,10 +467,10 @@ bool CreateFromSkeletalMeshInternal(UPhysicsAsset* PhysicsAsset, USkeletalMesh*
 
 			if (bShowProgress && IsInGameThread())
 			{
-				SlowTask.EnterProgressFrame(1.0f, FText::Format(NSLOCTEXT("PhysicsAssetEditor", "ResetCollsionStepInfo", "Generating collision for {0}"), FText::FromName(BoneName)));
+				SlowTask.EnterProgressFrame(1.0f, FText::Format(NSLOCTEXT(
+					"PhysicsAssetEditor", "ResetCollsionStepInfo", "Generating collision for {0}"), FText::FromName(BoneName)));
 			}
 
-
 			//Construct the info - in the case of merged bones we append all the data
 			FBoneVertInfo Info = Infos[BoneIndex];
 			FMergedBoneIndices BoneIndices;
@@ -475,7 +494,8 @@ bool CreateFromSkeletalMeshInternal(UPhysicsAsset* PhysicsAsset, USkeletalMesh*
 				check(NewBodySetup->BoneName == BoneName);
 
 				// Fill in collision info for this bone.
-				const bool bSuccess = CreateCollisionFromBoneInternal(NewBodySetup, SkelMesh, BoneIndex, Params, Info, TriangleCache);
+				const bool bSuccess = CreateCollisionFromBoneInternal(
+					NewBodySetup, SkelMesh, BoneIndex, Params, Info, TriangleCache, PhysicsAsset, PhysicalBoneIndices);
 				if (bSuccess)
 				{
 					// create joint to parent body
@@ -592,7 +612,9 @@ bool CreateFromSkeletalMesh(UPhysicsAsset* PhysicsAsset, USkeletalMesh* SkelMesh
 
 	FSkinnedBoneTriangleCache TriangleCache(*SkelMesh, Params);
 
-	if (Params.GeomType == EFG_SingleConvexHull || Params.GeomType == EFG_MultiConvexHull || Params.GeomType == EFG_LevelSet || Params.GeomType == EFG_SkinnedLevelSet)
+	if (Params.bAlwaysUseVertices || Params.GeomType == EFG_SingleConvexHull ||
+		Params.GeomType == EFG_MultiConvexHull || Params.GeomType == EFG_LevelSet ||
+		Params.GeomType == EFG_SkinnedLevelSet)
 	{
 		TriangleCache.BuildCache();
 	}
@@ -673,14 +695,15 @@ FMatrix ComputeViewCovarianceMatrix(TArrayView<const UE::Math::TVector<RealType>
 	return Covariance;
 }
 
-FMatrix ComputeCovarianceMatrix(const FBoneVertInfo& VertInfo)
+FMatrix ComputeCovarianceMatrix(const TArray<FVector3f>& Verts)
 {
-	return ComputeViewCovarianceMatrix(MakeArrayView(VertInfo.Positions));
+	return ComputeViewCovarianceMatrix(MakeArrayView(Verts));
 }
 
 FVector ComputeEigenVector(const FMatrix& A)
 {
-	//using the power method: this is ok because we only need the dominate eigenvector and speed is not critical: http://en.wikipedia.org/wiki/Power_iteration
+	// Using the power method: this is ok because we only need the dominant eigenvector and speed is
+	// not critical: http://en.wikipedia.org/wiki/Power_iteration
 	FVector Bk = FVector(0, 0, 1);
 	for (int32 i = 0; i < 32; ++i)
 	{
@@ -694,64 +717,176 @@ FVector ComputeEigenVector(const FMatrix& A)
 	return Bk.GetSafeNormal();
 }
 
-bool CreateCollisionFromBoneInternal(UBodySetup* bs, USkeletalMesh* skelMesh, int32 BoneIndex, const FPhysAssetCreateParams& Params, const FBoneVertInfo& Info, const FSkinnedBoneTriangleCache& TriangleCache)
+/**
+ * Collects BoneIndex plus all descendant bone indices in the skeleton hierarchy, stopping at any 
+ * bone that has its own physics body in the PhysicsAsset.
+ * 
+ * Note that this may return an empty array (OutBoneIndices) if the skeletal mesh, physics asset, 
+ * reference skeleton etc turn out to be invalid/problematic.
+ */
+static void GetRelevantBoneIndices(
+	const USkeletalMesh*   SkeletalMesh, 
+	const UPhysicsAsset*   PhysicsAsset, 
+	const int32            BoneIndex, 
+	TArray<int32>&         OutBoneIndices,
+	TConstArrayView<int32> BoneIndicesInCreation)
 {
-#if WITH_EDITOR
-	// multi convex hull can fail so wait to clear it ( will be called in DecomposeMeshToHulls() )
-	const bool bCanFail = (Params.GeomType == EFG_SingleConvexHull || Params.GeomType == EFG_MultiConvexHull || Params.GeomType == EFG_LevelSet);
-	if (!bCanFail)
+	if (!SkeletalMesh || !PhysicsAsset || BoneIndex == INDEX_NONE)
+	{
+		return;
+	}
+
+	const FReferenceSkeleton& RefSkeleton = SkeletalMesh->GetRefSkeleton();
+	const int32 NumBones = RefSkeleton.GetNum();
+
+	if (BoneIndex < 0 || BoneIndex >= NumBones)
+	{
+		return;
+	}
+
+	// Build a set of bone indices that have their own physics body in the asset.
+	// We look up by bone name since that's the reliable link between skeleton and physics asset.
+	TSet<int32> PhysicalBoneIndices;
+	for (const USkeletalBodySetup* BodySetup : PhysicsAsset->SkeletalBodySetups)
+	{
+		if (BodySetup)
+		{
+			int32 BodyBoneIdx = RefSkeleton.FindBoneIndex(BodySetup->BoneName);
+			if (BodyBoneIdx != INDEX_NONE)
+			{
+				PhysicalBoneIndices.Add(BodyBoneIdx);
+			}
+		}
+	}
+
+	// Duplicates are automatically ignored
+	PhysicalBoneIndices.Append(BoneIndicesInCreation);
+
+	// Iterative traversal down the hierarchy. Always include the starting bone (it's the one we're
+	// generating the hull for). For children, include them only if they are NOT physical, then
+	// continue recursing.
+	OutBoneIndices.Reset();
+	TArray<int32> Stack;
+	Stack.Push(BoneIndex);
+
+	while (Stack.Num() > 0)
 	{
-		// Empty any existing collision.
-		bs->RemoveSimpleCollision();
+		const int32 Current = Stack.Pop(EAllowShrinking::No);
+		OutBoneIndices.Add(Current);
+
+		// Find all direct children of Current in the ref skeleton.
+		for (int32 ChildIdx = 0; ChildIdx < NumBones; ++ChildIdx)
+		{
+			if (RefSkeleton.GetParentIndex(ChildIdx) == Current)
+			{
+				// Stop at (don't include) bones that have their own physics body,
+				// UNLESS it's the starting bone itself (which is physical by definition).
+				if (!PhysicalBoneIndices.Contains(ChildIdx))
+				{
+					Stack.Push(ChildIdx);
+				}
+			}
+		}
 	}
-#endif // WITH_EDITOR
+}
 
+bool CreateCollisionFromBoneInternal(
+	UBodySetup* BodySetup, USkeletalMesh* SkelMesh, int32 BoneIndex, const FPhysAssetCreateParams& Params, 
+	const FBoneVertInfo& Info, const FSkinnedBoneTriangleCache& TriangleCache, const UPhysicsAsset* PhysicsAsset,
+	TConstArrayView<int32> BoneIndicesInCreation)
+{
 	// Calculate orientation of to use for collision primitive.
 	FMatrix ElemTM = FMatrix::Identity;
-	bool ComputeFromVerts = false;
+
+	const TArray<FVector3f>* VertsToUse = &Info.Positions;
+	TArray<FVector3f> TriangleCacheVerts;
+	TArray<uint32> TriangleCacheIndices;
+
+	if (TriangleCache.IsBuilt())
+	{
+		TArray<int32> BoneIndices;
+		if (Params.bIncludeChildBones)
+		{
+			GetRelevantBoneIndices(SkelMesh, PhysicsAsset, BoneIndex, BoneIndices, BoneIndicesInCreation);
+		}
+		else
+		{
+			BoneIndices.Add(BoneIndex);
+		}
+
+		TArray<uint32> Unused;
+		TriangleCache.GetVerticesAndIndicesForBones(BoneIndex, BoneIndices, TriangleCacheVerts, TriangleCacheIndices, Unused);
+
+		if (TriangleCacheVerts.IsEmpty())
+		{
+			if (BodySetup->AggGeom.GetElementCount() > 0)
+			{
+				// If we're supposed to be using the mesh but have failed to find any verts, then
+				// just leave any existing collision as it is.
+				FMessageLog EditorErrors("EditorErrors");
+				FFormatNamedArguments Arguments;
+				Arguments.Add(TEXT("BoneName"), FText::FromName(SkelMesh->GetRefSkeleton().GetBoneName(BoneIndex)));
+				EditorErrors.Warning(FText::Format(NSLOCTEXT("PhysicsAssetUtils", "NoVerts",
+					"Unable to update hull for {BoneName} as there are no vertices associated with the bone."), Arguments));
+				EditorErrors.Open();
+				return false;
+			}
+		}
+		else
+		{
+			VertsToUse = &TriangleCacheVerts;
+		}
+	}
+
+	// Convex hulls aren't user-editable after creation, so don't set them unless we know the result
+	// will be good. Other shapes can be tweaked by the user, so even if we don't have vertices,
+	// create/change to them as the user can fix them up.
+	if (Params.GeomType == EFG_SingleConvexHull || Params.GeomType == EFG_MultiConvexHull)
+	{
+		if (VertsToUse->IsEmpty())
+		{
+			// If there are no verts, then either leave the shape as is, or create something "default"
+			FMessageLog EditorErrors("EditorErrors");
+			FFormatNamedArguments Arguments;
+			Arguments.Add(TEXT("BoneName"), FText::FromName(SkelMesh->GetRefSkeleton().GetBoneName(BoneIndex)));
+			if (BodySetup->AggGeom.GetElementCount() > 0)
+			{
+				EditorErrors.Warning(FText::Format(NSLOCTEXT("PhysicsAssetUtils", "NoVerts",
+					"Unable to update hull for {BoneName} as there are no vertices associated with the bone."), Arguments));
+				EditorErrors.Open();
+				return false;
+			}
+			EditorErrors.Warning(FText::Format(NSLOCTEXT("PhysicsAssetUtils", "NoVerts",
+				"Unable to create a good shape for the {BoneName} as there are no vertices associated with it."), Arguments));
+			EditorErrors.Open();
+		}
+	}
 
 	if (Params.bAutoOrientToBone)
 	{
-		// Compute covariance matrix for verts of this bone
-		// Then use axis with largest variance for orienting bone box
-		const FMatrix CovarianceMatrix = ComputeCovarianceMatrix(Info);
+		const FMatrix CovarianceMatrix = ComputeCovarianceMatrix(*VertsToUse);
 		FVector ZAxis = ComputeEigenVector(CovarianceMatrix);
 		FVector XAxis, YAxis;
 		ZAxis.FindBestAxisVectors(YAxis, XAxis);
 		ElemTM = FMatrix(XAxis, YAxis, ZAxis, FVector::ZeroVector);
 	}
 
-	// convert to FTransform now
-	// Matrix inverse doesn't handle well when DET == 0, so 
-	// convert to FTransform and use that data
+	// Convert to ElemTM to FTransform as Matrix inverse doesn't handle well when DET == 0, so 
 	FTransform ElementTransform(ElemTM);
 	// Get the (Unreal scale) bounding box for this bone using the rotation.
 	FBox BoneBox(ForceInit);
-	for (int32 j = 0; j < Info.Positions.Num(); j++)
+	for (const FVector3f& Vert : *VertsToUse)
 	{
-		BoneBox += ElementTransform.InverseTransformPosition((FVector)Info.Positions[j]);
+		BoneBox += ElementTransform.InverseTransformPosition(FVector(Vert));
 	}
 
+	// Verts are already in the bone's space (including scale)
 	FVector BoxCenter(0, 0, 0), BoxExtent(0, 0, 0);
-
-	FBox TransformedBox = BoneBox;
 	if (BoneBox.IsValid)
 	{
-		// make sure to apply scale to the box size
-		FMatrix BoneMatrix = skelMesh->GetComposedRefPoseMatrix(BoneIndex);
-		TransformedBox = BoneBox.TransformBy(FTransform(BoneMatrix));
 		BoneBox.GetCenterAndExtents(BoxCenter, BoxExtent);
 	}
-
-	float MinRad = TransformedBox.GetExtent().GetMin();
-	float MinAllowedSize = MinPrimSize;
-
-	// If the primitive is going to be too small - just use some default numbers and let the user tweak.
-	if (MinRad < MinAllowedSize)
-	{
-		// change min allowed size to be min, not DefaultPrimSize
-		BoxExtent = FVector(MinAllowedSize, MinAllowedSize, MinAllowedSize);
-	}
+	BoxExtent = FVector::Max(BoxExtent, FVector(MinPrimSize, MinPrimSize, MinPrimSize));
 
 	FVector BoneOrigin = ElementTransform.TransformPosition(BoxCenter);
 	ElementTransform.SetTranslation(BoneOrigin);
@@ -767,7 +902,9 @@ bool CreateCollisionFromBoneInternal(UBodySetup* bs, USkeletalMesh* skelMesh, in
 		BoxElem.Y = BoxExtent.Y * 2.0f * 1.01f;
 		BoxElem.Z = BoxExtent.Z * 2.0f * 1.01f;
 
-		bs->AggGeom.BoxElems.Add(BoxElem);
+		BodySetup->RemoveSimpleCollision();
+		BodySetup->AggGeom.BoxElems.Add(BoxElem);
+		BodySetup->InvalidatePhysicsData();
 	}
 	else if (Params.GeomType == EFG_Sphere)
 	{
@@ -776,30 +913,32 @@ bool CreateCollisionFromBoneInternal(UBodySetup* bs, USkeletalMesh* skelMesh, in
 		SphereElem.Center = ElementTransform.GetTranslation();
 		SphereElem.Radius = BoxExtent.GetMax() * 1.01f;
 
-		bs->AggGeom.SphereElems.Add(SphereElem);
+		BodySetup->RemoveSimpleCollision();
+		BodySetup->AggGeom.SphereElems.Add(SphereElem);
+		BodySetup->InvalidatePhysicsData();
 	}
 	else if (Params.GeomType == EFG_SingleConvexHull || Params.GeomType == EFG_MultiConvexHull)
 	{
-		TArray<FVector3f> Verts;
-		TArray<uint32> Indices;
-		TriangleCache.GetVerticesAndIndicesForBone(BoneIndex, Verts, Indices);
-
-		if (Verts.Num())
+		// We need to pick up vertices from this bone and any bone below it in the hierarchy,
+		// stopping the search when we move on to another bone with physics.
+		if (!TriangleCacheIndices.IsEmpty())
 		{
 			const int32 HullCount = Params.GeomType == EFG_SingleConvexHull ? 1 : Params.HullCount;
-			DecomposeMeshToHulls(bs, Verts, Indices, HullCount, Params.MaxHullVerts);
+			DecomposeMeshToHulls(BodySetup, *VertsToUse, TriangleCacheIndices, HullCount, Params.MaxHullVerts);
 		}
 		else
 		{
 			FMessageLog EditorErrors("EditorErrors");
-			EditorErrors.Warning(NSLOCTEXT("PhysicsAssetUtils", "ConvexNoPositions", "Unable to create a convex hull for the given bone as there are no vertices associated with the bone."));
+			FFormatNamedArguments Arguments;
+			Arguments.Add(TEXT("BoneName"), FText::FromName(SkelMesh->GetRefSkeleton().GetBoneName(BoneIndex)));
+			EditorErrors.Warning(FText::Format(NSLOCTEXT("PhysicsAssetUtils", "ConvexNoPositions",
+				"Unable to create a convex hull for {BoneName} as there are no vertices associated with the bone."), Arguments));
 			EditorErrors.Open();
 			return false;
 		}
 	}
 	else if (Params.GeomType == EFG_Sphyl)
 	{
-
 		FKSphylElem SphylElem;
 
 		if (BoxExtent.X > BoxExtent.Z && BoxExtent.X > BoxExtent.Y)
@@ -826,7 +965,9 @@ bool CreateCollisionFromBoneInternal(UBodySetup* bs, USkeletalMesh* skelMesh, in
 			SphylElem.Length = BoxExtent.Z * 1.01f;
 		}
 
-		bs->AggGeom.SphylElems.Add(SphylElem);
+		BodySetup->RemoveSimpleCollision();
+		BodySetup->AggGeom.SphylElems.Add(SphylElem);
+		BodySetup->InvalidatePhysicsData();
 	}
 	else if (Params.GeomType == EFG_TaperedCapsule)
 	{
@@ -857,17 +998,16 @@ bool CreateCollisionFromBoneInternal(UBodySetup* bs, USkeletalMesh* skelMesh, in
 			TaperedCapsuleElem.Length = BoxExtent.Z * 1.01f;
 		}
 
-		bs->AggGeom.TaperedCapsuleElems.Add(TaperedCapsuleElem);
+		BodySetup->RemoveSimpleCollision();
+		BodySetup->AggGeom.TaperedCapsuleElems.Add(TaperedCapsuleElem);
+		BodySetup->InvalidatePhysicsData();
 	}
 	else if (Params.GeomType == EFG_LevelSet)
 	{
-		TArray<FVector3f> Verts;
-		TArray<uint32> Indices;
-		TriangleCache.GetVerticesAndIndicesForBone(BoneIndex, Verts, Indices);
-
-		if (Verts.Num())
+		if (!VertsToUse->IsEmpty() && !TriangleCacheIndices.IsEmpty())
 		{
-			const bool bOK = LevelSetHelpers::CreateLevelSetForBone(bs, Verts, Indices, Params.LevelSetResolution);
+			// This clears out existing collision (or, if it fails, it leaves it)
+			const bool bOK = LevelSetHelpers::CreateLevelSetForBone(BodySetup, *VertsToUse, TriangleCacheIndices, Params.LevelSetResolution);
 			if (!bOK)
 			{
 				FMessageLog EditorErrors("EditorErrors");
@@ -888,27 +1028,36 @@ bool CreateCollisionFromBoneInternal(UBodySetup* bs, USkeletalMesh* skelMesh, in
 	return true;
 }
 
-bool CreateCollisionFromBone(UBodySetup* bs, USkeletalMesh* skelMesh, int32 BoneIndex, const FPhysAssetCreateParams& Params, const FBoneVertInfo& Info)
+bool CreateCollisionFromBone(
+	UBodySetup* BodySetup, USkeletalMesh* SkelMesh, int32 BoneIndex, const FPhysAssetCreateParams& Params, 
+	const FBoneVertInfo& Info, const UPhysicsAsset* PhysicsAsset, TConstArrayView<int32> BoneIndicesInCreation)
 {
-	check(skelMesh);
+	check(SkelMesh);
 
-	FSkinnedBoneTriangleCache TriangleCache(*skelMesh, Params);
+	FSkinnedBoneTriangleCache TriangleCache(*SkelMesh, Params);
 
-	if (Params.GeomType == EFG_SingleConvexHull || Params.GeomType == EFG_MultiConvexHull || Params.GeomType == EFG_LevelSet)
+	if (Params.bAlwaysUseVertices || Params.GeomType == EFG_SingleConvexHull ||
+		Params.GeomType == EFG_MultiConvexHull || Params.GeomType == EFG_LevelSet ||
+		Params.GeomType == EFG_SkinnedLevelSet)
 	{
 		TriangleCache.BuildCache();
 	}
 
-	return CreateCollisionFromBoneInternal(bs, skelMesh, BoneIndex, Params, Info, TriangleCache);
+	return CreateCollisionFromBoneInternal(
+		BodySetup, SkelMesh, BoneIndex, Params, Info, TriangleCache, PhysicsAsset, BoneIndicesInCreation);
 }
 
-bool CreateCollisionFromBones(UBodySetup* bs, USkeletalMesh* skelMesh, const TArray<int32>& BoneIndices, const FPhysAssetCreateParams& Params, const FBoneVertInfo& Info)
+bool CreateCollisionFromBones(
+	UBodySetup* BodySetup, USkeletalMesh* SkelMesh, const TArray<int32>& BoneIndices, const FPhysAssetCreateParams& Params, 
+	const FBoneVertInfo& Info, const UPhysicsAsset* PhysicsAsset)
 {
-	check(skelMesh);
+	check(SkelMesh);
 
-	FSkinnedBoneTriangleCache TriangleCache(*skelMesh, Params);
+	FSkinnedBoneTriangleCache TriangleCache(*SkelMesh, Params);
 
-	if (Params.GeomType == EFG_SingleConvexHull || Params.GeomType == EFG_MultiConvexHull || Params.GeomType == EFG_LevelSet)
+	if (Params.bAlwaysUseVertices || Params.GeomType == EFG_SingleConvexHull ||
+		Params.GeomType == EFG_MultiConvexHull || Params.GeomType == EFG_LevelSet ||
+		Params.GeomType == EFG_SkinnedLevelSet)
 	{
 		TriangleCache.BuildCache();
 	}
@@ -916,7 +1065,8 @@ bool CreateCollisionFromBones(UBodySetup* bs, USkeletalMesh* skelMesh, const TAr
 	bool bAllSuccessful = true;
 	for ( int Index = 0; Index < BoneIndices.Num(); ++Index )
 	{
-		if ( !CreateCollisionFromBoneInternal(bs, skelMesh, BoneIndices[Index], Params, Info, TriangleCache) )
+		if ( !CreateCollisionFromBoneInternal(
+			BodySetup, SkelMesh, BoneIndices[Index], Params, Info, TriangleCache, PhysicsAsset, BoneIndices) )
 		{
 			bAllSuccessful = false;
 		}
@@ -925,13 +1075,18 @@ bool CreateCollisionFromBones(UBodySetup* bs, USkeletalMesh* skelMesh, const TAr
 	return bAllSuccessful;
 }
 
-bool CreateCollisionsFromBones(UPhysicsAsset* PhysicsAsset, USkeletalMesh* SkelMesh, const TArray<int32>& BodyIndices, const FPhysAssetCreateParams& Params, const TArray<FBoneVertInfo>& Info, TArray<int32>& OutSuccessfulBodyIndices)
+bool CreateCollisionsFromBones(
+	UPhysicsAsset* PhysicsAsset, USkeletalMesh* SkelMesh, const TArray<int32>& BodyIndices, 
+	const FPhysAssetCreateParams& Params, const TArray<FBoneVertInfo>& Info, TArray<int32>& OutSuccessfulBodyIndices)
 {
 	check(SkelMesh);
 	check(PhysicsAsset);
 
 	FSkinnedBoneTriangleCache TriangleCache(*SkelMesh, Params);
-	if (Params.GeomType == EFG_SingleConvexHull || Params.GeomType == EFG_MultiConvexHull || Params.GeomType == EFG_LevelSet || Params.GeomType == EFG_SkinnedLevelSet)
+
+	if (Params.bAlwaysUseVertices || Params.GeomType == EFG_SingleConvexHull || 
+		Params.GeomType == EFG_MultiConvexHull || Params.GeomType == EFG_LevelSet || 
+		Params.GeomType == EFG_SkinnedLevelSet)
 	{
 		TriangleCache.BuildCache();
 	}
@@ -1003,7 +1158,6 @@ bool CreateCollisionsFromBones(UPhysicsAsset* PhysicsAsset, USkeletalMesh* SkelM
 		SlowTask.MakeDialog();
 	}
 
-
 	TArray<FName> SuccessfulBodySetupNames;
 	for (const FName& BoneName : BodySetupBoneNames)
 	{
@@ -1035,14 +1189,30 @@ bool CreateCollisionsFromBones(UPhysicsAsset* PhysicsAsset, USkeletalMesh* SkelM
 		}
 		else
 		{
-			if (CreateCollisionFromBoneInternal(BodySetup, SkelMesh, BoneIndex, Params, Info[BoneIndex], TriangleCache))
+			TArray<int32> BoneIndicesInCreation;
+			BoneIndicesInCreation.Reserve(BodyIndices.Num());
+			for (int32 Index = 0; Index < BodySetupBoneNames.Num(); ++Index)
+			{
+				const int32 BoneIndexInCreation = SkelMesh->GetRefSkeleton().FindBoneIndex(BodySetupBoneNames[Index]);
+				BoneIndicesInCreation.Add(BoneIndexInCreation);
+			}
+
+			if (CreateCollisionFromBoneInternal(
+				BodySetup, SkelMesh, BoneIndex, Params, Info[BoneIndex], TriangleCache, PhysicsAsset, BoneIndicesInCreation))
 			{
 				SuccessfulBodySetupNames.AddUnique(BoneName);
 			}
 			else
 			{
 				bAllSuccessful = false;
-				DestroyBody(PhysicsAsset, BodyIndex);
+
+				// If collision generation failed, then there are two possibilities:
+				// 1. It was a fresh generation, in which case we have a body with no shapes, which is Bad and must go.
+				// 2. It was a regeneration, in which case we have the unchanged shapes from before - keep them.
+				if (BodySetup->AggGeom.GetElementCount() == 0)
+				{
+					DestroyBody(PhysicsAsset, BodyIndex);
+				}
 			}
 		}
 	}
@@ -1078,14 +1248,28 @@ bool CreateCollisionsFromBones(UPhysicsAsset* PhysicsAsset, USkeletalMesh* SkelM
 			UBodySetup* NewBodySetup = PhysicsAsset->SkeletalBodySetups[NewBodyIndex];
 			check(NewBodySetup->BoneName == BoneName);
 			const int32 BoneIndex = SkelMesh->GetRefSkeleton().FindBoneIndex(NewBodySetup->BoneName);
-			if (CreateCollisionFromBoneInternal(NewBodySetup, SkelMesh, BoneIndex, Params, Info[BoneIndex], TriangleCache))
+
+
+			TArray<int32> BoneIndicesInCreation;
+			BoneIndicesInCreation.Reserve(BodyIndices.Num());
+			for (int32 Index = 0; Index < BodySetupBoneNames.Num(); ++Index)
+			{
+				const int32 BoneIndexInCreation = SkelMesh->GetRefSkeleton().FindBoneIndex(BodySetupBoneNames[Index]);
+				BoneIndicesInCreation.Add(BoneIndexInCreation);
+			}
+
+			if (CreateCollisionFromBoneInternal(
+				NewBodySetup, SkelMesh, BoneIndex, Params, Info[BoneIndex], TriangleCache, PhysicsAsset, BoneIndicesInCreation))
 			{
 				SuccessfulBodySetupNames.AddUnique(BoneName);
 			}
 			else
 			{
 				bAllSuccessful = false;
-				DestroyBody(PhysicsAsset, NewBodyIndex);
+				if (NewBodySetup->AggGeom.GetElementCount() == 0)
+				{
+					DestroyBody(PhysicsAsset, NewBodyIndex);
+				}
 			}
 		}
 	}
diff --git a/Engine/Source/Developer/PhysicsUtilities/Private/SkinnedBoneTriangleCache.h b/Engine/Source/Developer/PhysicsUtilities/Private/SkinnedBoneTriangleCache.h
--- a/Engine/Source/Developer/PhysicsUtilities/Private/SkinnedBoneTriangleCache.h
+++ b/Engine/Source/Developer/PhysicsUtilities/Private/SkinnedBoneTriangleCache.h
@@ -37,6 +37,8 @@ public:
 	// Vertex positions will be transformed by PrimaryBoneIndex's transform matrix.
 	void GetVerticesAndIndicesForBones(const int32 PrimaryBoneIndex, const TArray<int32>& BoneIndices, TArray<FVector3f>& OutVertexPositions, TArray<uint32>& OutIndices, TArray<uint32>& OutLocalToSkinnedVertexIndex) const;
 
+	bool IsBuilt() const { return !BoneIndexToInfluencedVertices.IsEmpty() || !BoneIndexToTriangles.IsEmpty(); }
+
 private:
 	typedef int32 FBoneIndex;
 	typedef uint32 FSkinnedVertexIndex;
diff --git a/Engine/Source/Developer/PhysicsUtilities/Public/PhysicsAssetUtils.h b/Engine/Source/Developer/PhysicsUtilities/Public/PhysicsAssetUtils.h
--- a/Engine/Source/Developer/PhysicsUtilities/Public/PhysicsAssetUtils.h
+++ b/Engine/Source/Developer/PhysicsUtilities/Public/PhysicsAssetUtils.h
@@ -47,6 +47,8 @@ struct FPhysAssetCreateParams
 		MinWeldSize = KINDA_SMALL_NUMBER;
 		GeomType = EFG_Sphyl;
 		VertWeight = EVW_DominantWeight;
+		bAlwaysUseVertices = true;
+		bIncludeChildBones = true;
 		bAutoOrientToBone = true;
 		bCreateConstraints = true;
 		bWalkPastSmall = true;
@@ -76,6 +78,20 @@ struct FPhysAssetCreateParams
 	UPROPERTY(EditAnywhere, Category = "Body Creation", meta=(DisplayName="Vertex Weighting Type"))
 	TEnumAsByte<EPhysAssetFitVertWeight> VertWeight;
 
+	/** 
+	 * If true, then mesh vertices are always used to calculate the collision shape, which can be slow. If false then 
+	 * the simpler shapes will not use vertices, which can result in faster processing.
+	 */
+	UPROPERTY(EditAnywhere, Category = "Body Creation")
+	bool								bAlwaysUseVertices;
+
+	/**
+	 * If true, then when mesh vertices are used, non-physical bones of the current bone will be included when
+	 * calculating the collision shapes.
+	 */
+	UPROPERTY(EditAnywhere, Category = "Body Creation")
+	bool								bIncludeChildBones;
+
 	/** Whether to automatically orient the created bodies to their corresponding bones */
 	UPROPERTY(EditAnywhere, Category = "Body Creation")
 	bool								bAutoOrientToBone;
@@ -148,25 +164,32 @@ namespace FPhysicsAssetUtils
 	 * If you find you're calling this function repeatedly for different bone indices on the same mesh,
 	 * CreateFromSkeletalMesh or CreateCollisionsFromBones will provide better performance.
 	 *
-	 * @param	bs					BodySetup to create the collision for
-	 * @param	skelMesh			The SkeletalMesh we create collision for
-	 * @param	BoneIndex			Index of the bone the collision is created for
-	 * @param	Params				Additional parameters to control the creation 
-	 * @param	Info				The vertices to create the collision for
+	 * @param	BodySetup              BodySetup to create the collision for
+	 * @param	SkelMesh               The SkeletalMesh we create collision for
+	 * @param	BoneIndex              Index of the bone the collision is created for
+	 * @param	Params                 Additional parameters to control the creation 
+	 * @param   Info                   The vertices to create the collision for
+	 * @param   PhysicsAsset           If set, then this can/will be used to recursively find relevant child bones for collision generation
+	 * @param   BoneIndicesInCreation  If set, then this can/will be used to identify bones that will be physical (for collision generation) 
+	 *                                 and should not be included in another bone's collision
 	 * @return  Returns true if successfully created collision from bone
 	 */
-	PHYSICSUTILITIES_API bool CreateCollisionFromBone( UBodySetup* bs, USkeletalMesh* skelMesh, int32 BoneIndex, const FPhysAssetCreateParams& Params, const FBoneVertInfo& Info );
+	PHYSICSUTILITIES_API bool CreateCollisionFromBone(
+		UBodySetup* BodySetup, USkeletalMesh* SkelMesh, int32 BoneIndex, const FPhysAssetCreateParams& Params,
+		const FBoneVertInfo& Info, const UPhysicsAsset* PhysicsAsset,
+		TConstArrayView<int32> BoneIndicesInCreation = TConstArrayView<int32>());
 
 	/** Replaces any collision already in the BodySetup with an auto-generated one using the parameters provided.
 	 * 
-	 * @param	bs					BodySetup to create the collision for
-	 * @param	skelMesh			The SkeletalMesh we create collision for
+	 * @param	BodySetup					BodySetup to create the collision for
+	 * @param	SkelMesh			The SkeletalMesh we create collision for
 	 * @param	BoneIndices			Indices of the bones the collisions are created for
 	 * @param	Params				Additional parameters to control the creation 
 	 * @param	Info				The vertices to create the collision for
+	 * @param   PhysicsAsset        If set, then this can/will be used to recursively find relevant child bones for collision generation
 	 * @return  Returns true if successfully created collision from all specified bones
 	 */
-	PHYSICSUTILITIES_API bool CreateCollisionFromBones( UBodySetup* bs, USkeletalMesh* skelMesh, const TArray<int32>& BoneIndices, FPhysAssetCreateParams& Params, const FBoneVertInfo& Info );
+	PHYSICSUTILITIES_API bool CreateCollisionFromBones( UBodySetup* BodySetup, USkeletalMesh* SkelMesh, const TArray<int32>& BoneIndices, FPhysAssetCreateParams& Params, const FBoneVertInfo& Info, const UPhysicsAsset* PhysicsAsset );
 
 	/** Replaces any collision already in the  with an auto-generated ones using the parameters provided.
 	 *
diff --git a/Engine/Source/Editor/EditorStyle/Private/StarshipStyle.cpp b/Engine/Source/Editor/EditorStyle/Private/StarshipStyle.cpp
--- a/Engine/Source/Editor/EditorStyle/Private/StarshipStyle.cpp
+++ b/Engine/Source/Editor/EditorStyle/Private/StarshipStyle.cpp
@@ -8049,6 +8049,8 @@ void FStarshipEditorStyle::FStyle::SetupToolkitStyles()
 		Set( "PhysicsAssetEditor.CreateBodyWithSphere", new IMAGE_BRUSH("PhysicsAssetEditor/icon_PhAT_Sphere_40x", Icon40x40));
 		Set( "PhysicsAssetEditor.CreateBodyWithSphyl", new IMAGE_BRUSH("PhysicsAssetEditor/icon_PhAT_Sphyl_40x", Icon40x40));
 		Set( "PhysicsAssetEditor.CreateBodyWithBox", new IMAGE_BRUSH("PhysicsAssetEditor/icon_PhAT_Box_40x", Icon40x40));
+		Set( "PhysicsAssetEditor.CreateBodyWithSingleConvex", new IMAGE_BRUSH("PhysicsAssetEditor/icon_PhAT_Convex_40x", Icon40x40));
+		Set( "PhysicsAssetEditor.CreateBodyWithMultiConvex", new IMAGE_BRUSH("PhysicsAssetEditor/icon_PhAT_MultiConvex_40x", Icon40x40));
 		Set( "PhysicsAssetEditor.CreateBodyWithTaperedCapsule", new IMAGE_BRUSH("PhysicsAssetEditor/icon_PhAT_TaperedCapsule_40x", Icon40x40));
 		Set( "PhysicsAssetEditor.CreateBodyShouldCreateConstraints", new IMAGE_BRUSH("PhysicsAssetEditor/Constraint_16x", Icon40x40));
 		Set( "PhysicsAssetEditor.AddNewBody", new IMAGE_BRUSH( "PhysicsAssetEditor/icon_PhAT_NewBody_40x", Icon40x40 ) );
diff --git a/Engine/Source/Editor/PhysicsAssetEditor/Private/PhysicsAssetEditor.cpp b/Engine/Source/Editor/PhysicsAssetEditor/Private/PhysicsAssetEditor.cpp
--- a/Engine/Source/Editor/PhysicsAssetEditor/Private/PhysicsAssetEditor.cpp
+++ b/Engine/Source/Editor/PhysicsAssetEditor/Private/PhysicsAssetEditor.cpp
@@ -218,6 +218,8 @@ static void FillCreateBodiesConstraintsMenu(FMenuBuilder& InSubMenuBuilder)
 		InSubMenuBuilder.AddMenuEntry( PhysicsAssetEditorCommands.CreateBodyWithBox );
 		InSubMenuBuilder.AddMenuEntry( PhysicsAssetEditorCommands.CreateBodyWithSphere );
 		InSubMenuBuilder.AddMenuEntry( PhysicsAssetEditorCommands.CreateBodyWithSphyl );
+		InSubMenuBuilder.AddMenuEntry( PhysicsAssetEditorCommands.CreateBodyWithSingleConvex );
+		InSubMenuBuilder.AddMenuEntry( PhysicsAssetEditorCommands.CreateBodyWithMultiConvex );
 		InSubMenuBuilder.AddMenuEntry( PhysicsAssetEditorCommands.CreateBodyWithTaperedCapsule );
 		InSubMenuBuilder.EndSection();
 	}
@@ -237,7 +239,7 @@ void FPhysicsAssetEditor::InitPhysicsAssetEditor(const EToolkitMode::Type Mode,
 	SharedData = MakeShareable(new FPhysicsAssetEditorSharedData());
 
 	SharedData->SelectionChangedEvent.AddRaw(this, &FPhysicsAssetEditor::HandleViewportSelectionChanged);
-	SharedData->HierarchyChangedEvent.AddRaw(this, &FPhysicsAssetEditor::RefreshHierachyTree);
+	SharedData->HierarchyChangedEvent.AddRaw(this, &FPhysicsAssetEditor::RefreshHierarchyTree);
 	SharedData->PreviewChangedEvent.AddRaw(this, &FPhysicsAssetEditor::RefreshPreviewViewport);
 	SharedData->PhysicsAsset = ObjectToEdit;
 
@@ -426,7 +428,7 @@ void FPhysicsAssetEditor::HandleViewportSelectionChanged(const TArray<FPhysicsAs
 	}
 }
 
-void FPhysicsAssetEditor::RefreshHierachyTree()
+void FPhysicsAssetEditor::RefreshHierarchyTree()
 {
 	if(SkeletonTree.IsValid())
 	{
@@ -504,7 +506,7 @@ void FPhysicsAssetEditor::PostUndo(bool bSuccess)
 	OnPostUndo.Broadcast();
 
 	SharedData->PostUndo();
-	RefreshHierachyTree();
+	RefreshHierarchyTree();
 
 	SharedData->RefreshPhysicsAssetChange(SharedData->PhysicsAsset);
 }
@@ -522,7 +524,7 @@ void FPhysicsAssetEditor::PostRedo( bool bSuccess )
 void FPhysicsAssetEditor::OnAssetReimport(UObject* Object)
 {
 	RecreatePhysicsState();
-	RefreshHierachyTree();
+	RefreshHierarchyTree();
 	RefreshPreviewViewport();
 
 	if (SharedData->EditorSkelComp && SharedData->EditorSkelComp->GetSkeletalMeshAsset())
@@ -1117,22 +1119,32 @@ void FPhysicsAssetEditor::BindCommands()
 
 	ToolkitCommands->MapAction(
 		Commands.CreateBodyWithBox,
-		FExecuteAction::CreateSPLambda(this, [this]() { this->CreateBodiesAndConstraintsForSelectedBones(EAggCollisionShape::Box, this->ShouldCreateConstraintsWhenCreatingBodies()); }),
+		FExecuteAction::CreateSPLambda(this, [this]() { this->CreateBodiesAndConstraintsForSelectedBones(EPhysAssetFitGeomType::EFG_Box, this->ShouldCreateConstraintsWhenCreatingBodies()); }),
 		FCanExecuteAction::CreateSP(this, &FPhysicsAssetEditor::IsNotSimulation));
 
 	ToolkitCommands->MapAction(
 		Commands.CreateBodyWithSphere,
-		FExecuteAction::CreateSPLambda(this, [this]() { this->CreateBodiesAndConstraintsForSelectedBones(EAggCollisionShape::Sphere, this->ShouldCreateConstraintsWhenCreatingBodies()); }),
+		FExecuteAction::CreateSPLambda(this, [this]() { this->CreateBodiesAndConstraintsForSelectedBones(EPhysAssetFitGeomType::EFG_Sphere, this->ShouldCreateConstraintsWhenCreatingBodies()); }),
 		FCanExecuteAction::CreateSP(this, &FPhysicsAssetEditor::IsNotSimulation));
 
 	ToolkitCommands->MapAction(
 		Commands.CreateBodyWithSphyl,
-		FExecuteAction::CreateSPLambda(this, [this]() { this->CreateBodiesAndConstraintsForSelectedBones(EAggCollisionShape::Sphyl, this->ShouldCreateConstraintsWhenCreatingBodies()); }),
+		FExecuteAction::CreateSPLambda(this, [this]() { this->CreateBodiesAndConstraintsForSelectedBones(EPhysAssetFitGeomType::EFG_Sphyl, this->ShouldCreateConstraintsWhenCreatingBodies()); }),
+		FCanExecuteAction::CreateSP(this, &FPhysicsAssetEditor::IsNotSimulation));
+
+	ToolkitCommands->MapAction(
+		Commands.CreateBodyWithSingleConvex,
+		FExecuteAction::CreateSPLambda(this, [this]() { this->CreateBodiesAndConstraintsForSelectedBones(EPhysAssetFitGeomType::EFG_SingleConvexHull, this->ShouldCreateConstraintsWhenCreatingBodies()); }),
+		FCanExecuteAction::CreateSP(this, &FPhysicsAssetEditor::IsNotSimulation));
+
+	ToolkitCommands->MapAction(
+		Commands.CreateBodyWithMultiConvex,
+		FExecuteAction::CreateSPLambda(this, [this]() { this->CreateBodiesAndConstraintsForSelectedBones(EPhysAssetFitGeomType::EFG_MultiConvexHull, this->ShouldCreateConstraintsWhenCreatingBodies()); }),
 		FCanExecuteAction::CreateSP(this, &FPhysicsAssetEditor::IsNotSimulation));
 
 	ToolkitCommands->MapAction(
 		Commands.CreateBodyWithTaperedCapsule,
-		FExecuteAction::CreateSPLambda(this, [this]() { this->CreateBodiesAndConstraintsForSelectedBones(EAggCollisionShape::TaperedCapsule, this->ShouldCreateConstraintsWhenCreatingBodies()); }),
+		FExecuteAction::CreateSPLambda(this, [this]() { this->CreateBodiesAndConstraintsForSelectedBones(EPhysAssetFitGeomType::EFG_TaperedCapsule, this->ShouldCreateConstraintsWhenCreatingBodies()); }),
 		FCanExecuteAction::CreateSP(this, &FPhysicsAssetEditor::IsNotSimulation));
 
 	ToolkitCommands->MapAction(
@@ -1872,7 +1884,7 @@ void FPhysicsAssetEditor::Mirror()
 	SharedData->Mirror();
 
 	RecreatePhysicsState();
-	RefreshHierachyTree();
+	RefreshHierarchyTree();
 	RefreshPreviewViewport();
 }
 
@@ -2320,7 +2332,7 @@ void FPhysicsAssetEditor::CreateOrConvertConstraint(EPhysicsAssetEditorConstrain
 	}
 
 	RecreatePhysicsState();
-	RefreshHierachyTree();
+	RefreshHierarchyTree();
 	RefreshPreviewViewport();
 }
 
@@ -2507,7 +2519,7 @@ void FPhysicsAssetEditor::AddNewPrimitive(EAggCollisionShape::Type InPrimitiveTy
 	SharedData->SetSelectedPrimitives(NewSelection);
 
 	RecreatePhysicsState();
-	RefreshHierachyTree();
+	RefreshHierarchyTree();
 	RefreshPreviewViewport();
 }
 
@@ -2555,7 +2567,7 @@ void FPhysicsAssetEditor::SetBodiesBelowPhysicsType( EPhysicsType InPhysicsType,
 	}
 
 	RecreatePhysicsState();
-	RefreshHierachyTree();
+	RefreshHierarchyTree();
 }
 
 bool FPhysicsAssetEditor::IsNotSimulation() const
@@ -2642,17 +2654,17 @@ void FPhysicsAssetEditor::OnChangeDefaultMesh(USkeletalMesh* OldPreviewMesh, USk
 		MeshUtilities.CalcBoneVertInfos(NewPreviewMesh, SharedData->DominantWeightBoneInfos, true);
 		MeshUtilities.CalcBoneVertInfos(NewPreviewMesh, SharedData->AnyWeightBoneInfos, false);
 
-		RefreshHierachyTree();
+		RefreshHierarchyTree();
 
 		SharedData->EditorSkelComp->SetDisablePostProcessBlueprint(true);
 	}
 }
 
-void FPhysicsAssetEditor::CreateBodiesAndConstraintsForSelectedBones(const EAggCollisionShape::Type InPrimitiveType, const bool bInShouldCreateConstraints)
+void FPhysicsAssetEditor::CreateBodiesAndConstraintsForSelectedBones(const EPhysAssetFitGeomType InGeomType, const bool bInShouldCreateConstraints)
 {
 	FPhysAssetCreateParams NewBodyData = GetDefault<UPhysicsAssetGenerationSettings>()->CreateParams;
 
-	NewBodyData.GeomType = ConvertAggCollisionShapeTypeToPhysicsAssetGeomType(InPrimitiveType);
+	NewBodyData.GeomType = InGeomType;
 	NewBodyData.bCreateConstraints = bInShouldCreateConstraints;
 
 	CreateBodiesAndConstraintsForSelectedBones(NewBodyData);
@@ -2698,6 +2710,10 @@ void FPhysicsAssetEditor::CreateBodiesAndConstraintsForSelectedBones(const FPhys
 			SharedData->AutoNameAllPrimitives(BodyIndex, NewBodyData.GeomType);
 		}
 
+		// SetSelectedBodies doesn't seem to broadcast/update correctly if BodyIndices is empty (it
+		// looks selected, but actually nothing is selected, so if we get called again we can end up
+		// regenerating everything!), so clear it to be sure.
+		SharedData->ClearSelectedBody();
 		SharedData->SetSelectedBodies(BodyIndices);
 	}
 	else
@@ -2713,18 +2729,31 @@ void FPhysicsAssetEditor::CreateBodiesAndConstraintsForSelectedBones(const FPhys
 
 			FScopedSlowTask SlowTask((float)BoneItems.Num());
 			SlowTask.MakeDialog();
-			for(TSharedPtr<ISkeletonTreeItem> BoneItem : BoneItems)
-			{
-				SlowTask.EnterProgressFrame(1.0f, FText::Format(LOCTEXT("ResetCollsionStepInfo", "Generating collision for {0}"), FText::FromName(BoneItem->GetRowItemName())));
 
-				UBoneProxy* BoneProxy = CastChecked<UBoneProxy>(BoneItem->GetObject());
+			TArray<int32> BoneIndices;
+			TArray<FName> BoneNames;
+
+			BoneIndices.Reserve(BoneItems.Num());
+			BoneNames.Reserve(BoneItems.Num());
 
+			const FReferenceSkeleton& RefSkeleton = EditorSkelMesh->GetRefSkeleton();
+			for (TSharedPtr<ISkeletonTreeItem> BoneItem : BoneItems)
+			{
+				UBoneProxy* BoneProxy = CastChecked<UBoneProxy>(BoneItem->GetObject());
 				const int32 BoneIndex = EditorSkelMesh->GetRefSkeleton().FindBoneIndex(BoneProxy->BoneName);
 				if (BoneIndex != INDEX_NONE)
 				{
-					SharedData->MakeNewBody(NewBodyData, BoneIndex);
+					BoneIndices.Add(BoneIndex);
+					BoneNames.Add(BoneItem->GetRowItemName());
 				}
 			}
+
+			for (int32 Index = 0 ; Index != BoneIndices.Num() ; ++Index)
+			{
+				SlowTask.EnterProgressFrame(1.0f, FText::Format(LOCTEXT(
+					"ResetCollsionStepInfo", "Generating collision for {0}"), FText::FromName(BoneNames[Index])));
+				SharedData->MakeNewBody(NewBodyData, BoneIndices[Index], true, BoneIndices);
+			}
 		}
 		else
 		{
@@ -2758,7 +2787,7 @@ void FPhysicsAssetEditor::CreateBodiesAndConstraintsForSelectedBones(const FPhys
 	RecreatePhysicsState();
 	SharedData->RefreshPhysicsAssetChange(SharedData->PhysicsAsset);
 	RefreshPreviewViewport();
-	RefreshHierachyTree();
+	RefreshHierarchyTree();
 }
 
 void FPhysicsAssetEditor::ShowNotificationMessage(const FText& Message, const SNotificationItem::ECompletionState CompletionState)
@@ -2877,7 +2906,7 @@ void FPhysicsAssetEditor::OnPasteProperties()
 	RecreatePhysicsState();
 	SharedData->RefreshPhysicsAssetChange(SharedData->PhysicsAsset);
 	RefreshPreviewViewport();
-	RefreshHierachyTree();
+	RefreshHierarchyTree();
 }
 
 bool FPhysicsAssetEditor::CanCopyProperties() const
@@ -3683,7 +3712,7 @@ void FPhysicsAssetEditor::OnAssetSelectedFromStaticMeshAssetPicker( const FAsset
 			}
 
 			SharedData->RefreshPhysicsAssetChange(SharedData->PhysicsAsset);
-			RefreshHierachyTree();
+			RefreshHierarchyTree();
 		}
 		else
 		{
@@ -4073,7 +4102,7 @@ bool FPhysicsAssetEditor::LoadMLLevelSetFromDataTable(UDataTable* DataTableMLLev
 									}
 									DestBody->AggGeom.MLLevelSetElems.Add(MLLevelsetElem);	
 									RecreatePhysicsState();
-									RefreshHierachyTree();
+									RefreshHierarchyTree();
 									RefreshPreviewViewport();
 								}
 							}
@@ -4152,7 +4181,7 @@ void FPhysicsAssetEditor::GenerateSkinnedTriangleMesh()
 	RecreatePhysicsState();
 	SharedData->RefreshPhysicsAssetChange(SharedData->PhysicsAsset);
 	RefreshPreviewViewport();
-	RefreshHierachyTree();
+	RefreshHierarchyTree();
 }
 
 void FPhysicsAssetEditor::ImportMLLevelSetFromDataTable()
diff --git a/Engine/Source/Editor/PhysicsAssetEditor/Private/PhysicsAssetEditor.h b/Engine/Source/Editor/PhysicsAssetEditor/Private/PhysicsAssetEditor.h
--- a/Engine/Source/Editor/PhysicsAssetEditor/Private/PhysicsAssetEditor.h
+++ b/Engine/Source/Editor/PhysicsAssetEditor/Private/PhysicsAssetEditor.h
@@ -78,7 +78,7 @@ public:
 	void HandleViewportSelectionChanged(const TArray<FPhysicsAssetEditorSelectedElement>& InSelectedElements);
 
 	/** Repopulates the hierarchy tree view */
-	void RefreshHierachyTree();
+	void RefreshHierarchyTree();
 
 	/** Refreshes the preview viewport */
 	void RefreshPreviewViewport();
@@ -136,7 +136,7 @@ public:
 
 	TSharedPtr<ISkeletonTree> GetSkeletonTree() const { return SkeletonTree.ToSharedRef(); }
 
-	void CreateBodiesAndConstraintsForSelectedBones(const EAggCollisionShape::Type InPrimitiveType, const bool bInShouldCreateConstraints);
+	void CreateBodiesAndConstraintsForSelectedBones(const EPhysAssetFitGeomType InGeomType, const bool bInShouldCreateConstraints);
 	void CreateBodiesAndConstraintsForSelectedBones(const FPhysAssetCreateParams& InBodyData);
 
 	/** Reset bone collision for selected or regenerate all bodies if no bodies are selected */
diff --git a/Engine/Source/Editor/PhysicsAssetEditor/Private/PhysicsAssetEditorActions.cpp b/Engine/Source/Editor/PhysicsAssetEditor/Private/PhysicsAssetEditorActions.cpp
--- a/Engine/Source/Editor/PhysicsAssetEditor/Private/PhysicsAssetEditorActions.cpp
+++ b/Engine/Source/Editor/PhysicsAssetEditor/Private/PhysicsAssetEditorActions.cpp
@@ -16,6 +16,8 @@ void FPhysicsAssetEditorCommands::RegisterCommands()
 	UI_COMMAND(CreateBodyWithBox, "Create Bodies with Box", "Create bodies with the Box primitive type (and optionally create constraints) for the selected bones using current generation settings (see the Tools tab)", EUserInterfaceActionType::Button, FInputChord());
 	UI_COMMAND(CreateBodyWithSphere, "Create Bodies with Sphere", "Create bodies with the Sphere primitive type (and optionally create constraints) for the selected bones using current generation settings (see the Tools tab)", EUserInterfaceActionType::Button, FInputChord());
 	UI_COMMAND(CreateBodyWithSphyl, "Create Bodies with Capsule", "Create bodies with the Capsule primitive type (and optionally create constraints) for the selected bones using current generation settings (see the Tools tab)", EUserInterfaceActionType::Button, FInputChord());
+	UI_COMMAND(CreateBodyWithSingleConvex, "Create Bodies with Single Convex Hull", "Create bodies with the single convex hull primitive type (and optionally create constraints) for the selected bones using current generation settings (see the Tools tab)", EUserInterfaceActionType::Button, FInputChord());
+	UI_COMMAND(CreateBodyWithMultiConvex, "Create Bodies with Multi Convex Hull", "Create bodies with the multiple convex hull primitive type (and optionally create constraints) for the selected bones using current generation settings (see the Tools tab)", EUserInterfaceActionType::Button, FInputChord());
 	UI_COMMAND(CreateBodyWithTaperedCapsule, "Create Bodies with Tapered Capsule (Clothing Only)", "Create bodies with the Tapered Capsule primitive type (and optionally create constraints) for the selected bones using current generation settings (see the Tools tab)", EUserInterfaceActionType::Button, FInputChord());
 	UI_COMMAND(CreateBodyShouldCreateConstraints, "Create Constraints", "Create constraints between any newly created bodies and the bodies on adjacent bones for each of the selected bones", EUserInterfaceActionType::ToggleButton, FInputChord());
 	UI_COMMAND(CreateOrRegenerateBodies, "Create/Regenerate Bodies", "Creates bodies and/or constraints (if none exist) or Regenerates bodies (if they do exist) for each of the selected bones using the current generation settings (see the Tools tab)", EUserInterfaceActionType::Button, FInputChord());
diff --git a/Engine/Source/Editor/PhysicsAssetEditor/Private/PhysicsAssetEditorActions.h b/Engine/Source/Editor/PhysicsAssetEditor/Private/PhysicsAssetEditorActions.h
--- a/Engine/Source/Editor/PhysicsAssetEditor/Private/PhysicsAssetEditorActions.h
+++ b/Engine/Source/Editor/PhysicsAssetEditor/Private/PhysicsAssetEditorActions.h
@@ -31,6 +31,8 @@ public:
 	TSharedPtr<FUICommandInfo> CreateBodyWithSphere;
 	TSharedPtr<FUICommandInfo> CreateBodyWithSphyl;
 	TSharedPtr<FUICommandInfo> CreateBodyWithBox;
+	TSharedPtr<FUICommandInfo> CreateBodyWithSingleConvex;
+	TSharedPtr<FUICommandInfo> CreateBodyWithMultiConvex;
 	TSharedPtr<FUICommandInfo> CreateBodyWithTaperedCapsule;
 	TSharedPtr<FUICommandInfo> CreateBodyShouldCreateConstraints;
 	TSharedPtr<FUICommandInfo> CreateOrRegenerateBodies;
diff --git a/Engine/Source/Editor/PhysicsAssetEditor/Private/PhysicsAssetEditorEditMode.cpp b/Engine/Source/Editor/PhysicsAssetEditor/Private/PhysicsAssetEditorEditMode.cpp
--- a/Engine/Source/Editor/PhysicsAssetEditor/Private/PhysicsAssetEditorEditMode.cpp
+++ b/Engine/Source/Editor/PhysicsAssetEditor/Private/PhysicsAssetEditorEditMode.cpp
@@ -1244,7 +1244,7 @@ void FPhysicsAssetEditorEditMode::HitNothing(FEditorViewportClient* InViewportCl
 	}
 
 	InViewportClient->Invalidate();
-	PhysicsAssetEditorPtr.Pin()->RefreshHierachyTree();
+	PhysicsAssetEditorPtr.Pin()->RefreshHierarchyTree();
 }
 
 #undef LOCTEXT_NAMESPACE
diff --git a/Engine/Source/Editor/PhysicsAssetEditor/Private/PhysicsAssetEditorSharedData.cpp b/Engine/Source/Editor/PhysicsAssetEditor/Private/PhysicsAssetEditorSharedData.cpp
--- a/Engine/Source/Editor/PhysicsAssetEditor/Private/PhysicsAssetEditorSharedData.cpp
+++ b/Engine/Source/Editor/PhysicsAssetEditor/Private/PhysicsAssetEditorSharedData.cpp
@@ -2710,12 +2710,16 @@ void FPhysicsAssetEditorSharedData::RecreateBody(const FPhysAssetCreateParams& B
 	RefreshPhysicsAssetChange(PhysicsAsset);
 }
 
-void FPhysicsAssetEditorSharedData::MakeNewBody(int32 NewBoneIndex, bool bAutoSelect)
+void FPhysicsAssetEditorSharedData::MakeNewBody(
+	int32 NewBoneIndex, bool bAutoSelect, TConstArrayView<int32> BoneIndicesInCreation)
 {
-	MakeNewBody(GetDefault<UPhysicsAssetGenerationSettings>()->CreateParams, NewBoneIndex, bAutoSelect);
+	MakeNewBody(
+		GetDefault<UPhysicsAssetGenerationSettings>()->CreateParams, NewBoneIndex, bAutoSelect, BoneIndicesInCreation);
 }
 
-void FPhysicsAssetEditorSharedData::MakeNewBody(const FPhysAssetCreateParams& NewBodyData, const int32 NewBoneIndex, const bool bAutoSelect)
+void FPhysicsAssetEditorSharedData::MakeNewBody(
+	const FPhysAssetCreateParams& NewBodyData, const int32 NewBoneIndex, const bool bAutoSelect,
+	TConstArrayView<int32> BoneIndicesInCreation)
 {
 	USkeletalMesh* EditorSkelMesh = PhysicsAsset->GetPreviewMesh();
 	if(EditorSkelMesh == nullptr)
@@ -2747,11 +2751,15 @@ void FPhysicsAssetEditorSharedData::MakeNewBody(const FPhysAssetCreateParams& Ne
 	// Create a new physics body for this bone.
 	if (NewBodyData.VertWeight == EVW_DominantWeight)
 	{
-		bCreatedBody = FPhysicsAssetUtils::CreateCollisionFromBone(BodySetup, EditorSkelMesh, NewBoneIndex, NewBodyData, DominantWeightBoneInfos[NewBoneIndex]);
+		bCreatedBody = FPhysicsAssetUtils::CreateCollisionFromBone(
+			BodySetup, EditorSkelMesh, NewBoneIndex, NewBodyData, 
+			DominantWeightBoneInfos[NewBoneIndex], PhysicsAsset, BoneIndicesInCreation);
 	}
 	else
 	{
-		bCreatedBody = FPhysicsAssetUtils::CreateCollisionFromBone(BodySetup, EditorSkelMesh, NewBoneIndex, NewBodyData, AnyWeightBoneInfos[NewBoneIndex]);
+		bCreatedBody = FPhysicsAssetUtils::CreateCollisionFromBone(
+			BodySetup, EditorSkelMesh, NewBoneIndex, NewBodyData, 
+			AnyWeightBoneInfos[NewBoneIndex], PhysicsAsset, BoneIndicesInCreation);
 	}
 
 	if (bCreatedBody == false)
diff --git a/Engine/Source/Editor/PhysicsAssetEditor/Private/PhysicsAssetEditorSharedData.h b/Engine/Source/Editor/PhysicsAssetEditor/Private/PhysicsAssetEditorSharedData.h
--- a/Engine/Source/Editor/PhysicsAssetEditor/Private/PhysicsAssetEditorSharedData.h
+++ b/Engine/Source/Editor/PhysicsAssetEditor/Private/PhysicsAssetEditorSharedData.h
@@ -232,8 +232,14 @@ public:
 	void ShowSelected();
 	void HideSelected();
 	void RefreshPhysicsAssetChange(const UPhysicsAsset* InPhysAsset, bool bFullClothRefresh = true);
-	void MakeNewBody(int32 NewBoneIndex, bool bAutoSelect = true);
-	void MakeNewBody(const FPhysAssetCreateParams& NewBodyData, const int32 NewBoneIndex, const bool bAutoSelect = true);
+	// The optional BoneIndicesInCreation can contain a list of bone indices that will be created -
+	// for when bones are checked to see if they are (or will be) physical.
+	void MakeNewBody(
+		int32 NewBoneIndex, bool bAutoSelect = true,
+		TConstArrayView<int32> BoneIndicesInCreation = TConstArrayView<int32>());
+	void MakeNewBody(
+		const FPhysAssetCreateParams& NewBodyData, const int32 NewBoneIndex, const bool bAutoSelect = true, 
+		TConstArrayView<int32> BoneIndicesInCreation = TConstArrayView<int32>());
 	void RecreateBody(const int32 BoneIndex, const bool bAutoSelect = true); // Create a new physics body to replace an existing one.
 	void RecreateBody(const FPhysAssetCreateParams& BodyData, const int32 BoneIndex, const bool bAutoSelect = true); // Create a new physics body to replace an existing one.
 	void MakeOrRecreateBody(const int32 NewBoneIndex, const bool bAutoSelect = true); // Create a new physics body, replacing an existing one at the same index if it exists.
diff --git a/Engine/Source/Editor/PhysicsAssetEditor/Private/PhysicsAssetEditorSkeletalMeshComponent.cpp b/Engine/Source/Editor/PhysicsAssetEditor/Private/PhysicsAssetEditorSkeletalMeshComponent.cpp
--- a/Engine/Source/Editor/PhysicsAssetEditor/Private/PhysicsAssetEditorSkeletalMeshComponent.cpp
+++ b/Engine/Source/Editor/PhysicsAssetEditor/Private/PhysicsAssetEditorSkeletalMeshComponent.cpp
@@ -258,8 +258,7 @@ FTransform UPhysicsAssetEditorSkeletalMeshComponent::GetPrimitiveTransform(const
 		return PrimTM * BoneTM;
 	}
 
-	// Should never reach here
-	check(0);
+	// Ideally we would never reach here. However, it can happen momentarily in the UI as shapes are regenerated
 	return FTransform::Identity;
 }
 
diff --git a/Engine/Source/Runtime/Engine/Classes/PhysicsEngine/AggregateGeom.h b/Engine/Source/Runtime/Engine/Classes/PhysicsEngine/AggregateGeom.h
--- a/Engine/Source/Runtime/Engine/Classes/PhysicsEngine/AggregateGeom.h
+++ b/Engine/Source/Runtime/Engine/Classes/PhysicsEngine/AggregateGeom.h
@@ -185,61 +185,62 @@ struct FKAggregateGeom
 		switch (Type)
 		{
 		case EAggCollisionShape::Sphere:
-			if (ensure(SphereElems.IsValidIndex(Index)))
+			if (SphereElems.IsValidIndex(Index))
 			{
 				return &SphereElems[Index];
 			}
 			break;
 		case EAggCollisionShape::Box:
-			if (ensure(BoxElems.IsValidIndex(Index)))
+			if (BoxElems.IsValidIndex(Index))
 			{
 				return &BoxElems[Index];
 			}
 			break;
 		case EAggCollisionShape::Sphyl:
-			if (ensure(SphylElems.IsValidIndex(Index)))
+			if (SphylElems.IsValidIndex(Index))
 			{
 				return &SphylElems[Index];
 			}
 			break;
 		case EAggCollisionShape::Convex:
-			if (ensure(ConvexElems.IsValidIndex(Index)))
+			if (ConvexElems.IsValidIndex(Index))
 			{
 				return &ConvexElems[Index];
 			}
 			break;
 		case EAggCollisionShape::TaperedCapsule:
-			if (ensure(TaperedCapsuleElems.IsValidIndex(Index)))
+			if (TaperedCapsuleElems.IsValidIndex(Index))
 			{
 				return &TaperedCapsuleElems[Index];
 			}
 			break;
 		case EAggCollisionShape::LevelSet:
-			if (ensure(LevelSetElems.IsValidIndex(Index)))
+			if (LevelSetElems.IsValidIndex(Index))
 			{
 				return &LevelSetElems[Index];
 			}
 			break;
 		case EAggCollisionShape::SkinnedLevelSet:
-			if (ensure(SkinnedLevelSetElems.IsValidIndex(Index)))
+			if (SkinnedLevelSetElems.IsValidIndex(Index))
 			{
 				return &SkinnedLevelSetElems[Index];
 			}
 			break;
 		case EAggCollisionShape::MLLevelSet:
-			if (ensure(MLLevelSetElems.IsValidIndex(Index)))
+			if (MLLevelSetElems.IsValidIndex(Index))
 			{
 				return &MLLevelSetElems[Index];
 			}
 			break;
 		case EAggCollisionShape::SkinnedTriangleMesh:
-			if (ensure(SkinnedTriangleMeshElems.IsValidIndex(Index)))
+			if (SkinnedTriangleMeshElems.IsValidIndex(Index))
 			{
 				return &SkinnedTriangleMeshElems[Index];
 			}
 			break;
 		default:
-			ensure(false);
+			// Failure is OK/returned/handled
+			break;
 		}
 		return nullptr;
 	}
@@ -249,61 +250,62 @@ struct FKAggregateGeom
 		switch (Type)
 		{
 		case EAggCollisionShape::Sphere:
-			if (ensure(SphereElems.IsValidIndex(Index)))
+			if (SphereElems.IsValidIndex(Index))
 			{
 				return &SphereElems[Index];
 			}
 			break;
 		case EAggCollisionShape::Box:
-			if (ensure(BoxElems.IsValidIndex(Index)))
+			if (BoxElems.IsValidIndex(Index))
 			{
 				return &BoxElems[Index];
 			}
 			break;
 		case EAggCollisionShape::Sphyl:
-			if (ensure(SphylElems.IsValidIndex(Index)))
+			if (SphylElems.IsValidIndex(Index))
 			{
 				return &SphylElems[Index];
 			}
 			break;
 		case EAggCollisionShape::Convex:
-			if (ensure(ConvexElems.IsValidIndex(Index)))
+			if (ConvexElems.IsValidIndex(Index))
 			{
 				return &ConvexElems[Index];
 			}
 			break;
 		case EAggCollisionShape::TaperedCapsule:
-			if (ensure(TaperedCapsuleElems.IsValidIndex(Index)))
+			if (TaperedCapsuleElems.IsValidIndex(Index))
 			{
 				return &TaperedCapsuleElems[Index];
 			}
 			break;
 		case EAggCollisionShape::LevelSet:
-			if (ensure(LevelSetElems.IsValidIndex(Index)))
+			if (LevelSetElems.IsValidIndex(Index))
 			{
 				return &LevelSetElems[Index];
 			}
 			break;
 		case EAggCollisionShape::SkinnedLevelSet:
-			if (ensure(SkinnedLevelSetElems.IsValidIndex(Index)))
+			if (SkinnedLevelSetElems.IsValidIndex(Index))
 			{
 				return &SkinnedLevelSetElems[Index];
 			}
 			break;
 		case EAggCollisionShape::MLLevelSet:
-			if (ensure(MLLevelSetElems.IsValidIndex(Index)))
+			if (MLLevelSetElems.IsValidIndex(Index))
 			{
 				return &MLLevelSetElems[Index];
 			}
 			break;
 		case EAggCollisionShape::SkinnedTriangleMesh:
-			if (ensure(SkinnedTriangleMeshElems.IsValidIndex(Index)))
+			if (SkinnedTriangleMeshElems.IsValidIndex(Index))
 			{
 				return &SkinnedTriangleMeshElems[Index];
 			}
 			break;
 		default:
-			ensure(false);
+			// Failure is OK/returned/handled
+			break;
 		}
 		return nullptr;
 	}
