Fix localization warnings

#Synced-CL 51208681

diff --git a/Engine/Source/Developer/PhysicsUtilities/Private/PhysicsAssetUtils.cpp b/Engine/Source/Developer/PhysicsUtilities/Private/PhysicsAssetUtils.cpp
--- a/Engine/Source/Developer/PhysicsUtilities/Private/PhysicsAssetUtils.cpp
+++ b/Engine/Source/Developer/PhysicsUtilities/Private/PhysicsAssetUtils.cpp
@@ -826,7 +826,7 @@ bool CreateCollisionFromBoneInternal(
 				FMessageLog EditorErrors("EditorErrors");
 				FFormatNamedArguments Arguments;
 				Arguments.Add(TEXT("BoneName"), FText::FromName(SkelMesh->GetRefSkeleton().GetBoneName(BoneIndex)));
-				EditorErrors.Warning(FText::Format(NSLOCTEXT("PhysicsAssetUtils", "NoVerts",
+				EditorErrors.Warning(FText::Format(NSLOCTEXT("PhysicsAssetUtils", "NoVertsUpdateHull2",
 					"Unable to update hull for {BoneName} as there are no vertices associated with the bone."), Arguments));
 				EditorErrors.Open();
 				return false;
@@ -851,12 +851,12 @@ bool CreateCollisionFromBoneInternal(
 			Arguments.Add(TEXT("BoneName"), FText::FromName(SkelMesh->GetRefSkeleton().GetBoneName(BoneIndex)));
 			if (BodySetup->AggGeom.GetElementCount() > 0)
 			{
-				EditorErrors.Warning(FText::Format(NSLOCTEXT("PhysicsAssetUtils", "NoVerts",
+				EditorErrors.Warning(FText::Format(NSLOCTEXT("PhysicsAssetUtils", "NoVertsUpdateHull1",
 					"Unable to update hull for {BoneName} as there are no vertices associated with the bone."), Arguments));
 				EditorErrors.Open();
 				return false;
 			}
-			EditorErrors.Warning(FText::Format(NSLOCTEXT("PhysicsAssetUtils", "NoVerts",
+			EditorErrors.Warning(FText::Format(NSLOCTEXT("PhysicsAssetUtils", "NoVertsCreateShape",
 				"Unable to create a good shape for the {BoneName} as there are no vertices associated with it."), Arguments));
 			EditorErrors.Open();
 		}
